Temp fix for ConcertReplicationScripting busting UE file path limit of 200 characters

diff --git a/Engine/Plugins/Developer/Concert/ConcertScripting/ConcertReplicationScripting/Source/ConcertReplicationScripting/ConcertReplicationScripting.Build.cs b/Engine/Plugins/Developer/Concert/ConcertScripting/ConcertReplicationScripting/Source/ConcertReplicationScripting/ConcertReplicationScripting.Build.cs
--- a/Engine/Plugins/Developer/Concert/ConcertScripting/ConcertReplicationScripting/Source/ConcertReplicationScripting/ConcertReplicationScripting.Build.cs
+++ b/Engine/Plugins/Developer/Concert/ConcertScripting/ConcertReplicationScripting/Source/ConcertReplicationScripting/ConcertReplicationScripting.Build.cs
@@ -27,7 +27,8 @@ namespace UnrealBuildTool.Rules
 				}
 			);
 
-			ShortName = "CrtSc";
+			// TODO UE-202216: Do proper fix
+			ShortName = "CS";
 		}
 	}
 }
\ No newline at end of file
diff --git a/Engine/Plugins/Developer/Concert/ConcertScripting/ConcertReplicationScripting/Source/ConcertReplicationScriptingEditor/ConcertReplicationScriptingEditor.build.cs b/Engine/Plugins/Developer/Concert/ConcertScripting/ConcertReplicationScripting/Source/ConcertReplicationScriptingEditor/ConcertReplicationScriptingEditor.build.cs
--- a/Engine/Plugins/Developer/Concert/ConcertScripting/ConcertReplicationScripting/Source/ConcertReplicationScriptingEditor/ConcertReplicationScriptingEditor.build.cs
+++ b/Engine/Plugins/Developer/Concert/ConcertScripting/ConcertReplicationScripting/Source/ConcertReplicationScriptingEditor/ConcertReplicationScriptingEditor.build.cs
@@ -34,8 +34,9 @@ namespace UnrealBuildTool.Rules
 					"ConcertTransport", // For LogConcert
 				}
 			);
-
-			ShortName = "CrtScEd";
+            
+			// TODO UE-202216: Do proper fix
+			ShortName = "CSE";
 		}
 	}
 }
