Fix non-unity in auto-generated reflection code

diff --git a/Engine/Plugins/Developer/Concert/ConcertScripting/ConcertReplicationScripting/Source/ConcertReplicationScripting/Private/ConcertReplicationBlueprintFunctionLibrary.cpp b/Engine/Plugins/Developer/Concert/ConcertScripting/ConcertReplicationScripting/Source/ConcertReplicationScripting/Private/ConcertReplicationBlueprintFunctionLibrary.cpp
--- a/Engine/Plugins/Developer/Concert/ConcertScripting/ConcertReplicationScripting/Source/ConcertReplicationScripting/Private/ConcertReplicationBlueprintFunctionLibrary.cpp
+++ b/Engine/Plugins/Developer/Concert/ConcertScripting/ConcertReplicationScripting/Source/ConcertReplicationScripting/Private/ConcertReplicationBlueprintFunctionLibrary.cpp
@@ -6,8 +6,6 @@
 #include "ConcertPropertyChainWrapper.h"
 #include "Replication/PropertyChainUtils.h"
 
-#include "Templates/SubclassOf.h"
-
 bool UConcertReplicationBlueprintFunctionLibrary::MakePropertyChainByLiteralPath(const TSubclassOf<UObject>& Class, const TArray<FName>& PathToProperty, FConcertPropertyChainWrapper& Result)
 {
 	if (!Class)
diff --git a/Engine/Plugins/Developer/Concert/ConcertScripting/ConcertReplicationScripting/Source/ConcertReplicationScripting/Public/ConcertReplicationBlueprintFunctionLibrary.h b/Engine/Plugins/Developer/Concert/ConcertScripting/ConcertReplicationScripting/Source/ConcertReplicationScripting/Public/ConcertReplicationBlueprintFunctionLibrary.h
--- a/Engine/Plugins/Developer/Concert/ConcertScripting/ConcertReplicationScripting/Source/ConcertReplicationScripting/Public/ConcertReplicationBlueprintFunctionLibrary.h
+++ b/Engine/Plugins/Developer/Concert/ConcertScripting/ConcertReplicationScripting/Source/ConcertReplicationScripting/Public/ConcertReplicationBlueprintFunctionLibrary.h
@@ -3,6 +3,7 @@
 #pragma once
 
 #include "Kismet/BlueprintFunctionLibrary.h"
+#include "Templates/SubclassOf.h"
 #include "ConcertReplicationBlueprintFunctionLibrary.generated.h"
 
 struct FConcertPropertyChainWrapper;
