Update comments about ShortName in response to marking ticket Won't Fix

diff --git a/Engine/Plugins/Developer/Concert/ConcertScripting/ConcertReplicationScripting/Source/ConcertReplicationScripting/ConcertReplicationScripting.Build.cs b/Engine/Plugins/Developer/Concert/ConcertScripting/ConcertReplicationScripting/Source/ConcertReplicationScripting/ConcertReplicationScripting.Build.cs
--- a/Engine/Plugins/Developer/Concert/ConcertScripting/ConcertReplicationScripting/Source/ConcertReplicationScripting/ConcertReplicationScripting.Build.cs
+++ b/Engine/Plugins/Developer/Concert/ConcertScripting/ConcertReplicationScripting/Source/ConcertReplicationScripting/ConcertReplicationScripting.Build.cs
@@ -27,7 +27,7 @@ namespace UnrealBuildTool.Rules
 				}
 			);
 
-			// TODO UE-202216: Do proper fix
+			// Hacky but we decided to keep it like this to avoid exceeding the 200 character path limit
 			ShortName = "CS";
 		}
 	}
diff --git a/Engine/Plugins/Developer/Concert/ConcertScripting/ConcertReplicationScripting/Source/ConcertReplicationScriptingEditor/ConcertReplicationScriptingEditor.build.cs b/Engine/Plugins/Developer/Concert/ConcertScripting/ConcertReplicationScripting/Source/ConcertReplicationScriptingEditor/ConcertReplicationScriptingEditor.build.cs
--- a/Engine/Plugins/Developer/Concert/ConcertScripting/ConcertReplicationScripting/Source/ConcertReplicationScriptingEditor/ConcertReplicationScriptingEditor.build.cs
+++ b/Engine/Plugins/Developer/Concert/ConcertScripting/ConcertReplicationScripting/Source/ConcertReplicationScriptingEditor/ConcertReplicationScriptingEditor.build.cs
@@ -35,7 +35,7 @@ namespace UnrealBuildTool.Rules
 				}
 			);
             
-			// TODO UE-202216: Do proper fix
+			// Hacky but we decided to keep it like this to avoid exceeding the 200 character path limit
 			ShortName = "CSE";
 		}
 	}
