Improve error log when property is invalid

diff --git a/Engine/Plugins/Developer/Concert/ConcertScripting/ConcertReplicationScripting/Source/ConcertReplicationScripting/Private/ConcertReplicationBlueprintFunctionLibrary.cpp b/Engine/Plugins/Developer/Concert/ConcertScripting/ConcertReplicationScripting/Source/ConcertReplicationScripting/Private/ConcertReplicationBlueprintFunctionLibrary.cpp
--- a/Engine/Plugins/Developer/Concert/ConcertScripting/ConcertReplicationScripting/Source/ConcertReplicationScripting/Private/ConcertReplicationBlueprintFunctionLibrary.cpp
+++ b/Engine/Plugins/Developer/Concert/ConcertScripting/ConcertReplicationScripting/Source/ConcertReplicationScripting/Private/ConcertReplicationBlueprintFunctionLibrary.cpp
@@ -31,8 +31,9 @@ bool UConcertReplicationBlueprintFunctionLibrary::MakePropertyChainByLiteralPath
 		return true;
 	}
 
-	UE_LOG(LogConcert, Error, TEXT("UConcertReplicationBlueprintFunctionLibrary::MakePropertyChainByLiteralPath: Path [%s] is invalid!"),
-		*FString::JoinBy(PathToProperty, TEXT(", "), [](const FName& Name){ return TEXT("\"") + Name.ToString() + TEXT("\""); })
+	UE_LOG(LogConcert, Error, TEXT("UConcertReplicationBlueprintFunctionLibrary::MakePropertyChainByLiteralPath: Path [%s] does not point to any property in class %s or it is not replicatable!"),
+		*FString::JoinBy(PathToProperty, TEXT(", "), [](const FName& Name){ return TEXT("\"") + Name.ToString() + TEXT("\""); }),
+		*Class->GetPathName()
 		);
 	return false;
 }
