Added GetAllProperties and GetChildProperties

diff --git a/Engine/Plugins/Developer/Concert/ConcertScripting/ConcertReplicationScripting/Source/ConcertReplicationScripting/Private/ConcertReplicationBlueprintFunctionLibrary.cpp b/Engine/Plugins/Developer/Concert/ConcertScripting/ConcertReplicationScripting/Source/ConcertReplicationScripting/Private/ConcertReplicationBlueprintFunctionLibrary.cpp
--- a/Engine/Plugins/Developer/Concert/ConcertScripting/ConcertReplicationScripting/Source/ConcertReplicationScripting/Private/ConcertReplicationBlueprintFunctionLibrary.cpp
+++ b/Engine/Plugins/Developer/Concert/ConcertScripting/ConcertReplicationScripting/Source/ConcertReplicationScripting/Private/ConcertReplicationBlueprintFunctionLibrary.cpp
@@ -42,7 +42,7 @@ TArray<FConcertPropertyChainWrapper> UConcertReplicationBlueprintFunctionLibrary
 {
 	if (!Class)
 	{
-		UE_LOG(LogConcert, Error, TEXT("UConcertReplicationBlueprintFunctionLibrary::GetPropertiesIn: Invalid class!"))
+		UE_LOG(LogConcert, Error, TEXT("UConcertReplicationBlueprintFunctionLibrary::GetPropertiesIn: No class set!"))
 		return {};
 	}
 
@@ -61,6 +61,47 @@ TArray<FConcertPropertyChainWrapper> UConcertReplicationBlueprintFunctionLibrary
 	return Result;
 }
 
+TArray<FConcertPropertyChainWrapper> UConcertReplicationBlueprintFunctionLibrary::GetAllProperties(const TSubclassOf<UObject>& Class)
+{
+	if (!Class)
+	{
+		UE_LOG(LogConcert, Error, TEXT("UConcertReplicationBlueprintFunctionLibrary::GetPropertiesIn: No class set!"))
+		return {};
+	}
+
+	TArray<FConcertPropertyChainWrapper> Result;  
+	UE::ConcertSyncCore::PropertyChain::ForEachReplicatableConcertProperty(
+		**Class,
+		[&Result](FConcertPropertyChain&& PropertyChain)
+		{
+			Result.Emplace(MoveTemp(PropertyChain));
+			return EBreakBehavior::Continue;
+		});
+	return Result;
+}
+
+TArray<FConcertPropertyChainWrapper> UConcertReplicationBlueprintFunctionLibrary::GetChildProperties(const FConcertPropertyChainWrapper& Parent, const TSubclassOf<UObject>& Class, bool bOnlyDirect)
+{
+	if (!Class)
+	{
+		UE_LOG(LogConcert, Error, TEXT("UConcertReplicationBlueprintFunctionLibrary::GetPropertiesIn: No class set!"))
+		return {};
+	}
+
+	TArray<FConcertPropertyChainWrapper> Result;  
+	UE::ConcertSyncCore::PropertyChain::ForEachReplicatableConcertProperty(
+		**Class,
+		[bOnlyDirect, &Result, &Parent](FConcertPropertyChain&& PropertyChain)
+		{
+			if ((bOnlyDirect && PropertyChain.IsDirectChildOf(Parent.PropertyChain)) || (!bOnlyDirect && PropertyChain.IsChildOf(Parent.PropertyChain)))
+			{
+				Result.Emplace(MoveTemp(PropertyChain));
+			}
+			return EBreakBehavior::Continue;
+		});
+	return Result;
+}
+
 const TArray<FName>& UConcertReplicationBlueprintFunctionLibrary::GetPropertyStringPath(const FConcertPropertyChainWrapper& Path)
 {
 	return Path.PropertyChain.GetPathToProperty();
diff --git a/Engine/Plugins/Developer/Concert/ConcertScripting/ConcertReplicationScripting/Source/ConcertReplicationScripting/Public/ConcertReplicationBlueprintFunctionLibrary.h b/Engine/Plugins/Developer/Concert/ConcertScripting/ConcertReplicationScripting/Source/ConcertReplicationScripting/Public/ConcertReplicationBlueprintFunctionLibrary.h
--- a/Engine/Plugins/Developer/Concert/ConcertScripting/ConcertReplicationScripting/Source/ConcertReplicationScripting/Public/ConcertReplicationBlueprintFunctionLibrary.h
+++ b/Engine/Plugins/Developer/Concert/ConcertScripting/ConcertReplicationScripting/Source/ConcertReplicationScripting/Public/ConcertReplicationBlueprintFunctionLibrary.h
@@ -40,6 +40,24 @@ public:
 	 */
 	UFUNCTION(BlueprintPure, Category = "Concert|Replication", meta = (KeyWords = "Get Properties In Find Enumerate List"))
 	static TArray<FConcertPropertyChainWrapper> GetPropertiesIn(const TSubclassOf<UObject>& Class, FPropertyChainPredicate Filter);
+
+	/**
+	 * Gets all properties in the class that are valid for replicating.
+	 * @param Class The class in which to search
+	 * @return All properties in the class that are valid for replicating.
+	 */
+	UFUNCTION(BlueprintPure, Category = "Concert|Replication")
+	static TArray<FConcertPropertyChainWrapper> GetAllProperties(const TSubclassOf<UObject>& Class);
+	
+	/**
+	 * Returns all child properties of Parent that are valid for replicating.
+	 * @param Parent The property of which to find child properties
+	 * @param Class The class in which to search
+	 * @param bOnlyDirect Whether you only want direct children of Parent
+	 * @return All child properties of Parent that are valid for replicating.
+	 */
+	UFUNCTION(BlueprintPure, Category = "Concert|Replication", meta = (ScriptMethod))
+	static TArray<FConcertPropertyChainWrapper> GetChildProperties(const FConcertPropertyChainWrapper& Parent, const TSubclassOf<UObject>& Class, bool bOnlyDirect = false);
 	
 	/******************** Getters ********************/
 	
