Add detail customization for FConcertPropertyChainWrapper

diff --git a/Engine/Plugins/Developer/Concert/ConcertScripting/ConcertReplicationScripting/Source/ConcertReplicationScripting/Private/ConcertReplicationBlueprintFunctionLibrary.cpp b/Engine/Plugins/Developer/Concert/ConcertScripting/ConcertReplicationScripting/Source/ConcertReplicationScripting/Private/ConcertReplicationBlueprintFunctionLibrary.cpp
--- a/Engine/Plugins/Developer/Concert/ConcertScripting/ConcertReplicationScripting/Source/ConcertReplicationScripting/Private/ConcertReplicationBlueprintFunctionLibrary.cpp
+++ b/Engine/Plugins/Developer/Concert/ConcertScripting/ConcertReplicationScripting/Source/ConcertReplicationScripting/Private/ConcertReplicationBlueprintFunctionLibrary.cpp
@@ -102,6 +102,11 @@ TArray<FConcertPropertyChainWrapper> UConcertReplicationBlueprintFunctionLibrary
 	return Result;
 }
 
+FString UConcertReplicationBlueprintFunctionLibrary::ToString(const FConcertPropertyChainWrapper& PropertyChain)
+{
+	return PropertyChain.PropertyChain.ToString();
+}
+
 const TArray<FName>& UConcertReplicationBlueprintFunctionLibrary::GetPropertyStringPath(const FConcertPropertyChainWrapper& Path)
 {
 	return Path.PropertyChain.GetPathToProperty();
diff --git a/Engine/Plugins/Developer/Concert/ConcertScripting/ConcertReplicationScripting/Source/ConcertReplicationScripting/Public/ConcertReplicationBlueprintFunctionLibrary.h b/Engine/Plugins/Developer/Concert/ConcertScripting/ConcertReplicationScripting/Source/ConcertReplicationScripting/Public/ConcertReplicationBlueprintFunctionLibrary.h
--- a/Engine/Plugins/Developer/Concert/ConcertScripting/ConcertReplicationScripting/Source/ConcertReplicationScripting/Public/ConcertReplicationBlueprintFunctionLibrary.h
+++ b/Engine/Plugins/Developer/Concert/ConcertScripting/ConcertReplicationScripting/Source/ConcertReplicationScripting/Public/ConcertReplicationBlueprintFunctionLibrary.h
@@ -58,6 +58,10 @@ public:
 	 */
 	UFUNCTION(BlueprintPure, Category = "Concert|Replication", meta = (ScriptMethod))
 	static TArray<FConcertPropertyChainWrapper> GetChildProperties(const FConcertPropertyChainWrapper& Parent, const TSubclassOf<UObject>& Class, bool bOnlyDirect = false);
+
+	/** Converts the the property to a string. */
+	UFUNCTION(BlueprintPure, Category = "Concert|Replication", meta = (ScriptMethod))
+	static FString ToString(const FConcertPropertyChainWrapper& PropertyChain);
 	
 	/******************** Getters ********************/
 	
diff --git a/Engine/Plugins/Developer/Concert/ConcertScripting/ConcertReplicationScripting/Source/ConcertReplicationScriptingEditor/ConcertReplicationScriptingEditor.build.cs b/Engine/Plugins/Developer/Concert/ConcertScripting/ConcertReplicationScripting/Source/ConcertReplicationScriptingEditor/ConcertReplicationScriptingEditor.build.cs
--- a/Engine/Plugins/Developer/Concert/ConcertScripting/ConcertReplicationScripting/Source/ConcertReplicationScriptingEditor/ConcertReplicationScriptingEditor.build.cs
+++ b/Engine/Plugins/Developer/Concert/ConcertScripting/ConcertReplicationScripting/Source/ConcertReplicationScriptingEditor/ConcertReplicationScriptingEditor.build.cs
@@ -26,12 +26,15 @@ namespace UnrealBuildTool.Rules
 				new string[]
 				{
 					// Engine
+					"InputCore",
 					"Slate",
 					"SlateCore",
-					"UnrealEd",
+					"PropertyEditor",
+					"UnrealEd", 
 					
 					// Concert
 					"ConcertTransport", // For LogConcert
+					"ConcertSharedSlate"
 				}
 			);
             
diff --git a/Engine/Plugins/Developer/Concert/ConcertScripting/ConcertReplicationScripting/Source/ConcertReplicationScriptingEditor/Private/ConcertReplicationScriptingEditorModule.cpp b/Engine/Plugins/Developer/Concert/ConcertScripting/ConcertReplicationScripting/Source/ConcertReplicationScriptingEditor/Private/ConcertReplicationScriptingEditorModule.cpp
--- a/Engine/Plugins/Developer/Concert/ConcertScripting/ConcertReplicationScripting/Source/ConcertReplicationScriptingEditor/Private/ConcertReplicationScriptingEditorModule.cpp
+++ b/Engine/Plugins/Developer/Concert/ConcertScripting/ConcertReplicationScripting/Source/ConcertReplicationScriptingEditor/Private/ConcertReplicationScriptingEditorModule.cpp
@@ -2,15 +2,29 @@
 
 #include "ConcertReplicationScriptingEditorModule.h"
 
+#include "ConcertPropertyChainWrapper.h"
+#include "ReplicationScriptingStyle.h"
+#include "Customization/ConcertPropertyCustomization.h"
 #include "Modules/ModuleManager.h"
+#include "PropertyEditorModule.h"
 
 namespace UE::ConcertReplicationScriptingEditor
 {
 	void FConcertReplicationScriptingEditorModule::StartupModule()
-	{}
+	{
+		FReplicationScriptingStyle::Initialize();
+		
+		FPropertyEditorModule& PropertyModule = FModuleManager::LoadModuleChecked<FPropertyEditorModule>("PropertyEditor");
+		PropertyModule.RegisterCustomPropertyTypeLayout(
+			FConcertPropertyChainWrapper::StaticStruct()->GetFName(),
+			FOnGetPropertyTypeCustomizationInstance::CreateStatic(&FConcertPropertyCustomization::MakeInstance)
+			);
+	}
 
 	void FConcertReplicationScriptingEditorModule::ShutdownModule()
-	{}
+	{
+		FReplicationScriptingStyle::Shutdown();
+	}
 }
 
 IMPLEMENT_MODULE(UE::ConcertReplicationScriptingEditor::FConcertReplicationScriptingEditorModule, ConcertReplicationScriptingEditor);
\ No newline at end of file
diff --git a/Engine/Plugins/Developer/Concert/ConcertScripting/ConcertReplicationScripting/Source/ConcertReplicationScriptingEditor/Private/Customization/ConcertPropertyCustomization.cpp b/Engine/Plugins/Developer/Concert/ConcertScripting/ConcertReplicationScripting/Source/ConcertReplicationScriptingEditor/Private/Customization/ConcertPropertyCustomization.cpp
new file mode 100644
--- /dev/null
+++ b/Engine/Plugins/Developer/Concert/ConcertScripting/ConcertReplicationScripting/Source/ConcertReplicationScriptingEditor/Private/Customization/ConcertPropertyCustomization.cpp
@@ -0,0 +1,113 @@
+// Copyright Epic Games, Inc. All Rights Reserved.
+
+#include "ConcertPropertyCustomization.h"
+
+#include "ConcertPropertyChainWrapper.h"
+#include "Widgets/SConcertPropertyChainCombo.h"
+
+#include "DetailWidgetRow.h"
+
+namespace UE::ConcertReplicationScriptingEditor
+{
+	/** Keeps track of the last class that was used for searching properties in SConcertPropertyChainCombo. */
+	static TWeakObjectPtr<const UClass> LastSelectedClass;
+	/**
+	 * Maps every property path to the last used class. This why when a user works with multiple properties, each one will remember the last used class for convenience.
+	 * This leaks memory but a user must edit a lot of properties to really be noticeable... should be ok.
+	 */
+	static TMap<FString, TWeakObjectPtr<const UClass>> PropertyToLastUsedClass;
+	
+	TSharedRef<IPropertyTypeCustomization> FConcertPropertyCustomization::MakeInstance()
+	{
+		return MakeShared<FConcertPropertyCustomization>();
+	}
+
+	void FConcertPropertyCustomization::CustomizeHeader(TSharedRef<IPropertyHandle> StructPropertyHandle, FDetailWidgetRow& HeaderRow, IPropertyTypeCustomizationUtils& StructCustomizationUtils)
+	{
+		PropertyHandle = StructPropertyHandle;
+		PropertyHandle->SetOnPropertyValueChanged(FSimpleDelegate::CreateSP(this, &FConcertPropertyCustomization::RefreshPropertiesAndUpdateUI));
+		RefreshProperties();
+		
+		const FString PropertyPath(PropertyHandle->GetPropertyPath());
+		const TWeakObjectPtr<const UClass>* LastBoundClass = PropertyToLastUsedClass.Find(PropertyPath);
+		const TWeakObjectPtr<const UClass> ClassToUse = LastBoundClass ? *LastBoundClass : LastSelectedClass;
+		
+		HeaderRow
+		.NameContent()
+		[
+			StructPropertyHandle->CreatePropertyNameWidget()
+		]
+		.ValueContent()
+		[
+			SNew(SBox)
+			.Padding(FMargin(0,2,0,1))
+			[
+				SAssignNew(PropertyComboBox, SConcertPropertyChainCombo)
+				.InitialClassSelection(ClassToUse.IsValid() ? ClassToUse.Get() : nullptr)
+				.IsEditable(!PropertyHandle->IsEditConst())
+				.ContainedProperties(&Properties)
+				.OnClassChanged_Lambda([this](const UClass* Class)
+				{
+					LastSelectedClass = Class;
+					
+					const FString PropertyPath(PropertyHandle->GetPropertyPath());
+					PropertyToLastUsedClass.Add(PropertyPath, Class);
+				})
+				.OnPropertySelectionChanged(this, &FConcertPropertyCustomization::OnPropertySelectionChanged)
+			]
+		];
+	}
+
+	void FConcertPropertyCustomization::RefreshProperties()
+	{
+		TArray<const void*> RawStructData;
+		PropertyHandle->AccessRawData(RawStructData);
+
+		Properties.Reset();
+		bHasMultipleValues = false;
+
+		TOptional<FConcertPropertyChain> SharedValue;
+		for (int32 Idx = 0; Idx < RawStructData.Num(); ++Idx)
+		{
+			if (!RawStructData[Idx])
+			{
+				continue;
+			}
+
+			const FConcertPropertyChainWrapper& CurrentValue = *(FConcertPropertyChainWrapper*)RawStructData[Idx];
+			Properties.Add(CurrentValue.PropertyChain);
+			
+			if (!SharedValue)
+			{
+				SharedValue = CurrentValue.PropertyChain;
+			}
+
+			const bool bValueIsDifferent = *SharedValue != CurrentValue.PropertyChain;
+			bHasMultipleValues |= bValueIsDifferent;
+		}
+	}
+
+	void FConcertPropertyCustomization::RefreshPropertiesAndUpdateUI()
+	{
+		RefreshProperties();
+		PropertyComboBox->RefreshPropertyContent();
+	}
+
+	void FConcertPropertyCustomization::OnPropertySelectionChanged(const FConcertPropertyChain& Property, bool bIsSelected)
+	{
+		// And this updates the underlying property data
+		TArray<void*> RawStructData;
+		PropertyHandle->AccessRawData(RawStructData);
+		FConcertPropertyChainWrapper NewValue = bIsSelected ? FConcertPropertyChainWrapper{ Property } : FConcertPropertyChainWrapper{};
+		for (int32 Idx = 0; Idx < RawStructData.Num(); ++Idx)
+		{
+			FConcertPropertyChainWrapper& Value = *(FConcertPropertyChainWrapper*)RawStructData[Idx];
+			Value = NewValue;
+		}
+		
+		// This makes sure the checkboxes display the correct state ...
+		Properties = { NewValue.PropertyChain };
+		// ... and this updates the combo button's content. 
+		PropertyComboBox->RefreshPropertyContent();
+	}
+}
diff --git a/Engine/Plugins/Developer/Concert/ConcertScripting/ConcertReplicationScripting/Source/ConcertReplicationScriptingEditor/Private/Customization/ConcertPropertyCustomization.h b/Engine/Plugins/Developer/Concert/ConcertScripting/ConcertReplicationScripting/Source/ConcertReplicationScriptingEditor/Private/Customization/ConcertPropertyCustomization.h
new file mode 100644
--- /dev/null
+++ b/Engine/Plugins/Developer/Concert/ConcertScripting/ConcertReplicationScripting/Source/ConcertReplicationScriptingEditor/Private/Customization/ConcertPropertyCustomization.h
@@ -0,0 +1,44 @@
+// Copyright Epic Games, Inc. All Rights Reserved.
+
+#pragma once
+
+#include "IPropertyTypeCustomization.h"
+#include "Widgets/SConcertPropertyChainCombo.h"
+
+struct FConcertPropertyChain;
+class IDetailChildrenBuilder;
+
+namespace UE::ConcertReplicationScriptingEditor
+{
+	/** Shows a combo button that allows selecting a property from a class. */
+	class FConcertPropertyCustomization : public IPropertyTypeCustomization
+	{
+	public:
+	
+		static TSharedRef<IPropertyTypeCustomization> MakeInstance();
+
+		//~ Begin IPropertyTypeCustomization Interface
+		virtual void CustomizeHeader(TSharedRef<IPropertyHandle> StructPropertyHandle, FDetailWidgetRow& HeaderRow, IPropertyTypeCustomizationUtils& StructCustomizationUtils) override;
+		virtual void CustomizeChildren(TSharedRef<IPropertyHandle> InStructPropertyHandle, IDetailChildrenBuilder& ChildBuilder, IPropertyTypeCustomizationUtils& StructCustomizationUtils) override {}
+		//~ End IPropertyTypeCustomization Interface
+
+	private:
+
+		/** Handle to the property containing the FConcertPropertyChainWrapper struct. */
+		TSharedPtr<IPropertyHandle> PropertyHandle;
+
+		/** The UI displayed for the property. */
+		TSharedPtr<SConcertPropertyChainCombo> PropertyComboBox;
+
+		/** Dummy passed in to the SConcertPropertyChainCombo. Holds all the values from the different property sources. */
+		TSet<FConcertPropertyChain> Properties;
+		bool bHasMultipleValues = false;
+
+		void RefreshProperties();
+		void RefreshPropertiesAndUpdateUI();
+		
+		void OnPropertySelectionChanged(const FConcertPropertyChain& Property, bool bIsSelected);
+	};
+}
+
+
diff --git a/Engine/Plugins/Developer/Concert/ConcertScripting/ConcertReplicationScripting/Source/ConcertReplicationScriptingEditor/Private/ReplicationScriptingStyle.cpp b/Engine/Plugins/Developer/Concert/ConcertScripting/ConcertReplicationScripting/Source/ConcertReplicationScriptingEditor/Private/ReplicationScriptingStyle.cpp
new file mode 100644
--- /dev/null
+++ b/Engine/Plugins/Developer/Concert/ConcertScripting/ConcertReplicationScripting/Source/ConcertReplicationScriptingEditor/Private/ReplicationScriptingStyle.cpp
@@ -0,0 +1,88 @@
+// Copyright Epic Games, Inc. All Rights Reserved.
+
+#include "ReplicationScriptingStyle.h"
+
+#include "Brushes/SlateBoxBrush.h"
+#include "Brushes/SlateImageBrush.h"
+#include "Brushes/SlateNoResource.h"
+#include "Brushes/SlateRoundedBoxBrush.h"
+#include "Layout/Margin.h"
+#include "Misc/Paths.h"
+#include "Styling/AppStyle.h"
+#include "Styling/SlateStyle.h"
+#include "Styling/SlateStyleRegistry.h"
+#include "Styling/SlateTypes.h"
+#include "Styling/StyleColors.h"
+
+namespace UE::ConcertReplicationScriptingEditor
+{
+	TSharedPtr<FSlateStyleSet> FReplicationScriptingStyle::StyleSet;
+
+	FName FReplicationScriptingStyle::GetStyleSetName()
+	{
+		return FName(TEXT("ReplicationScriptingStyle"));
+	}
+
+	void FReplicationScriptingStyle::Initialize()
+	{
+		// Only register once
+		if (StyleSet.IsValid())
+		{
+			return;
+		}
+
+		StyleSet = MakeShared<FSlateStyleSet>(GetStyleSetName());
+		StyleSet->SetContentRoot(FPaths::EngineContentDir() / TEXT("Editor/Slate"));
+		StyleSet->SetCoreContentRoot(FPaths::EngineContentDir() / TEXT("Slate"));
+
+		// This stuff is copied to look similar like the gameplay tag chips
+		const FLinearColor ChipColor = FStyleColors::Hover.GetSpecifiedColor();
+		const FLinearColor ChipColorHover = FStyleColors::Hover2.GetSpecifiedColor();
+		const FLinearColor ChipColorDisable = ChipColor.CopyWithNewOpacity(0.35);
+		FButtonStyle ChipButtonSelected = FAppStyle::Get().GetWidgetStyle<FButtonStyle>("Button");
+		ChipButtonSelected
+			.SetNormal(FSlateRoundedBoxBrush(ChipColor, 3.0f))
+			.SetHovered(FSlateRoundedBoxBrush(ChipColorHover, 3.0f))
+			.SetPressed(FSlateRoundedBoxBrush(ChipColorHover, 3.0f))
+			.SetDisabled(FSlateRoundedBoxBrush(ChipColorDisable, 3.0f))
+			.SetNormalForeground(FStyleColors::ForegroundHover)
+			.SetHoveredForeground(FStyleColors::White)
+			.SetPressedForeground(FStyleColors::White)
+			.SetDisabledForeground(FStyleColors::ForegroundHover)
+			.SetNormalPadding(FMargin(5,2,2,2))
+			.SetPressedPadding(FMargin(5,3,2,1));
+		StyleSet->Set("ConcertProperty.ChipButton.Selected", ChipButtonSelected);
+
+		FButtonStyle ChipClearButton = FAppStyle::Get().GetWidgetStyle<FButtonStyle>("Button");
+		ChipClearButton.SetNormal(FSlateNoResource())
+			.SetHovered(FSlateNoResource())
+			.SetPressed(FSlateNoResource())
+			.SetDisabled(FSlateNoResource())
+			.SetNormalForeground(FStyleColors::ForegroundHeader)
+			.SetHoveredForeground(FStyleColors::ForegroundHover)
+			.SetPressedForeground(FStyleColors::ForegroundHover)
+			.SetDisabledForeground(FStyleColors::Foreground)
+			.SetNormalPadding(FMargin(2,0,2,0))
+			.SetPressedPadding(FMargin(2,0,2,0));
+		StyleSet->Set("ConcertProperty.ChipClearButton", ChipClearButton);
+		
+		FSlateStyleRegistry::RegisterSlateStyle(*StyleSet.Get());
+	};
+
+	void FReplicationScriptingStyle::Shutdown()
+	{
+		if (StyleSet.IsValid())
+		{
+			FSlateStyleRegistry::UnRegisterSlateStyle(*StyleSet.Get());
+			ensure(StyleSet.IsUnique());
+			StyleSet.Reset();
+		}
+	}
+
+	TSharedPtr<ISlateStyle> FReplicationScriptingStyle::Get()
+	{
+		return StyleSet;
+	}
+}
+
+
diff --git a/Engine/Plugins/Developer/Concert/ConcertScripting/ConcertReplicationScripting/Source/ConcertReplicationScriptingEditor/Private/ReplicationScriptingStyle.h b/Engine/Plugins/Developer/Concert/ConcertScripting/ConcertReplicationScripting/Source/ConcertReplicationScriptingEditor/Private/ReplicationScriptingStyle.h
new file mode 100644
--- /dev/null
+++ b/Engine/Plugins/Developer/Concert/ConcertScripting/ConcertReplicationScripting/Source/ConcertReplicationScriptingEditor/Private/ReplicationScriptingStyle.h
@@ -0,0 +1,26 @@
+// Copyright Epic Games, Inc. All Rights Reserved.
+
+#pragma once
+
+#include "Styling/ISlateStyle.h"
+#include "Styling/SlateStyle.h"
+
+namespace UE::ConcertReplicationScriptingEditor
+{
+	class FReplicationScriptingStyle
+	{
+	public:
+	
+		static void Initialize();
+		static void Shutdown();
+
+		static TSharedPtr<ISlateStyle> Get();
+		static FName GetStyleSetName();
+	
+	private:
+	
+		static TSharedPtr<FSlateStyleSet> StyleSet;
+	};
+}
+
+
diff --git a/Engine/Plugins/Developer/Concert/ConcertScripting/ConcertReplicationScripting/Source/ConcertReplicationScriptingEditor/Private/Widgets/PropertyChainDelegates.h b/Engine/Plugins/Developer/Concert/ConcertScripting/ConcertReplicationScripting/Source/ConcertReplicationScriptingEditor/Private/Widgets/PropertyChainDelegates.h
new file mode 100644
--- /dev/null
+++ b/Engine/Plugins/Developer/Concert/ConcertScripting/ConcertReplicationScripting/Source/ConcertReplicationScriptingEditor/Private/Widgets/PropertyChainDelegates.h
@@ -0,0 +1,16 @@
+// Copyright Epic Games, Inc. All Rights Reserved.
+
+#pragma once
+
+#include "Delegates/Delegate.h"
+
+class UClass;
+struct FConcertPropertyChain;
+
+namespace UE::ConcertReplicationScriptingEditor
+{
+	DECLARE_DELEGATE(FOnPressed);
+	
+	DECLARE_DELEGATE_OneParam(FOnClassChanged, const UClass* Class);
+	DECLARE_DELEGATE_TwoParams(FOnSelectedPropertiesChanged, const FConcertPropertyChain& Property, bool bIsSelected);
+}
\ No newline at end of file
diff --git a/Engine/Plugins/Developer/Concert/ConcertScripting/ConcertReplicationScripting/Source/ConcertReplicationScriptingEditor/Private/Widgets/PropertySelectionColumn.cpp b/Engine/Plugins/Developer/Concert/ConcertScripting/ConcertReplicationScripting/Source/ConcertReplicationScriptingEditor/Private/Widgets/PropertySelectionColumn.cpp
new file mode 100644
--- /dev/null
+++ b/Engine/Plugins/Developer/Concert/ConcertScripting/ConcertReplicationScripting/Source/ConcertReplicationScriptingEditor/Private/Widgets/PropertySelectionColumn.cpp
@@ -0,0 +1,58 @@
+// Copyright Epic Games, Inc. All Rights Reserved.
+
+#include "PropertySelectionColumn.h"
+
+#include "Algo/AnyOf.h"
+#include "Replication/Editor/View/ReplicationColumnsUtils.h"
+
+#define LOCTEXT_NAMESPACE "PropertySelectionCheckboxColumn"
+
+namespace UE::ConcertReplicationScriptingEditor
+{
+	const FName PropertySelectionCheckboxColumnId(TEXT("AddPropertyCheckboxColumn"));
+	
+	ConcertSharedSlate::ReplicationColumns::FReplicationPropertyColumn MakePropertySelectionCheckboxColumn(
+		const TSet<FConcertPropertyChain>& SelectedProperties,
+		FOnSelectProperty OnSelectPropertyDelegate,
+		bool bIsEditable,
+		int32 SortPriority
+		)
+	{
+		using namespace ConcertSharedSlate;
+		using FPropertyColumnDelegates = TReplicationColumnDelegates<FReplicatedPropertyData>;
+		return MakeCheckboxColumn<FReplicatedPropertyData>(
+			PropertySelectionCheckboxColumnId,
+			FPropertyColumnDelegates(
+				FPropertyColumnDelegates::FGetColumnCheckboxState::CreateLambda(
+				[&SelectedProperties](const FReplicatedPropertyData& Data)
+				{
+					const FConcertPropertyChain& DisplayProperty = Data.GetProperty();
+					const bool bShouldBeSelected = SelectedProperties.Contains(DisplayProperty)
+						|| Algo::AnyOf(SelectedProperties, [&DisplayProperty](const FConcertPropertyChain& SelectedProperty)
+						{
+							return SelectedProperty.IsChildOf(DisplayProperty);
+						});
+					return bShouldBeSelected ? ECheckBoxState::Checked : ECheckBoxState::Unchecked;
+				}),
+				FPropertyColumnDelegates::FOnColumnCheckboxChanged::CreateLambda(
+				[OnSelectPropertyDelegate = MoveTemp(OnSelectPropertyDelegate)](bool bIsChecked, const FReplicatedPropertyData& Data)
+				{
+					OnSelectPropertyDelegate.Execute(Data.GetProperty(), bIsChecked);
+				}),
+				FPropertyColumnDelegates::FGetToolTipText::CreateLambda(
+				[](const FReplicatedPropertyData& Data)
+				{
+					return LOCTEXT("IncludePropertyTooltip", "Whether the property is included");
+				}),
+				FPropertyColumnDelegates::FIsEnabled::CreateLambda([bIsEditable](const FReplicatedPropertyData& Data)
+				{
+					return bIsEditable;
+				})),
+			FText::GetEmpty(),
+			SortPriority,
+			20.f
+		);
+	}
+}
+
+#undef LOCTEXT_NAMESPACE
\ No newline at end of file
diff --git a/Engine/Plugins/Developer/Concert/ConcertScripting/ConcertReplicationScripting/Source/ConcertReplicationScriptingEditor/Private/Widgets/PropertySelectionColumn.h b/Engine/Plugins/Developer/Concert/ConcertScripting/ConcertReplicationScripting/Source/ConcertReplicationScriptingEditor/Private/Widgets/PropertySelectionColumn.h
new file mode 100644
--- /dev/null
+++ b/Engine/Plugins/Developer/Concert/ConcertScripting/ConcertReplicationScripting/Source/ConcertReplicationScriptingEditor/Private/Widgets/PropertySelectionColumn.h
@@ -0,0 +1,27 @@
+// Copyright Epic Games, Inc. All Rights Reserved.
+
+#pragma once
+
+#include "Replication/Editor/View/PredefinedReplicationColumns.h"
+#include "Delegates/Delegate.h"
+
+namespace UE::ConcertReplicationScriptingEditor
+{
+	extern const FName PropertySelectionCheckboxColumnId;
+	
+	DECLARE_DELEGATE_TwoParams(FOnSelectProperty, const FConcertPropertyChain& PropertyChain, bool bIsSelected);
+
+	/**
+	 * Creates a checkbox column.
+	 *
+	 * @param SelectedProperties The checkbox is checked if the path is in SelectedProperties.
+	 * @param OnSelectPropertyDelegate Executed when a checkbox is checked or unchecked.
+	 */
+	ConcertSharedSlate::ReplicationColumns::FReplicationPropertyColumn MakePropertySelectionCheckboxColumn(
+		const TSet<FConcertPropertyChain>& SelectedProperties,
+		FOnSelectProperty OnSelectPropertyDelegate,
+		bool bIsEditable,
+		int32 SortPriority = 0 /* This places it before label column by default */
+		);
+}
+
diff --git a/Engine/Plugins/Developer/Concert/ConcertScripting/ConcertReplicationScripting/Source/ConcertReplicationScriptingEditor/Private/Widgets/SConcertPropertyChainChip.cpp b/Engine/Plugins/Developer/Concert/ConcertScripting/ConcertReplicationScripting/Source/ConcertReplicationScriptingEditor/Private/Widgets/SConcertPropertyChainChip.cpp
new file mode 100644
--- /dev/null
+++ b/Engine/Plugins/Developer/Concert/ConcertScripting/ConcertReplicationScripting/Source/ConcertReplicationScriptingEditor/Private/Widgets/SConcertPropertyChainChip.cpp
@@ -0,0 +1,84 @@
+// Copyright Epic Games, Inc. All Rights Reserved.
+
+#include "SConcertPropertyChainChip.h"
+
+#include "ReplicationScriptingStyle.h"
+#include "Replication/Data/ConcertPropertySelection.h"
+
+#include "Styling/AppStyle.h"
+#include "Styling/StyleColors.h"
+#include "Widgets/Images/SImage.h"
+#include "Widgets/Input/SButton.h"
+#include "Widgets/Layout/SWidgetSwitcher.h"
+#include "Widgets/SBoxPanel.h"
+#include "Widgets/Text/STextBlock.h"
+
+#define LOCTEXT_NAMESPACE "SConcertPropertyChainChip"
+
+namespace UE::ConcertReplicationScriptingEditor
+{
+	void SConcertPropertyChainChip::Construct(const FArguments& InArgs)
+	{
+		const FConcertPropertyChain PropertyChain = InArgs._DisplayedProperty.Get();
+		ChildSlot
+		[
+			SNew(SWidgetSwitcher)
+			.WidgetIndex(PropertyChain.IsEmpty() ? 1 : 0)
+
+			+SWidgetSwitcher::Slot()
+			[
+				SNew(SButton)
+				.ButtonStyle(FReplicationScriptingStyle::Get(), "ConcertProperty.ChipButton.Selected")
+				.OnClicked_Lambda([OnEditPressed = InArgs._OnEditPressed](){ OnEditPressed.ExecuteIfBound(); return FReply::Handled(); })
+				[
+					SNew(SHorizontalBox)
+
+					+SHorizontalBox::Slot()
+					.AutoWidth()
+					.VAlign(VAlign_Center)
+					[
+						SNew(STextBlock)
+						.Font(FAppStyle::GetFontStyle( TEXT("PropertyWindow.NormalFont")))
+						.Text(FText::FromString(PropertyChain.ToString()))
+					]
+
+					+SHorizontalBox::Slot()
+					.AutoWidth()
+					.VAlign(VAlign_Center)
+					[
+						SNew(SButton)
+						.Visibility(!PropertyChain.IsEmpty() && InArgs._ShowClearButton ? EVisibility::Visible : EVisibility::Collapsed)
+						.VAlign(VAlign_Center)
+						.HAlign(HAlign_Center)
+						.ButtonStyle(FReplicationScriptingStyle::Get(), "ConcertProperty.ChipClearButton")
+						.ToolTipText(LOCTEXT("Remove", "Remove"))
+						.ContentPadding(FMargin(2, 0, 0, 0))
+						.OnClicked_Lambda([OnClearPressed = InArgs._OnClearPressed]()
+						{
+							OnClearPressed.ExecuteIfBound();
+							return FReply::Handled();
+						})
+						[
+							SNew(SImage)
+							.ColorAndOpacity(FStyleColors::Foreground)
+							.Image(FAppStyle::GetBrush("Icons.X"))
+							.DesiredSizeOverride(FVector2D(12.0f, 12.0f))
+						]
+					]
+				]
+			]
+			
+			+SWidgetSwitcher::Slot()
+			[
+				SNew(STextBlock)
+				.ColorAndOpacity(FSlateColor::UseSubduedForeground())
+				.Font(FAppStyle::GetFontStyle( TEXT("PropertyWindow.NormalFont")))
+				.Text(LOCTEXT("EmptyPropertyChain.Label", "Empty"))
+				.ToolTipText(LOCTEXT("EmptyPropertyChain.Tooltip", "Empty property"))
+			]
+			
+		];
+	}
+}
+
+#undef LOCTEXT_NAMESPACE
\ No newline at end of file
diff --git a/Engine/Plugins/Developer/Concert/ConcertScripting/ConcertReplicationScripting/Source/ConcertReplicationScriptingEditor/Private/Widgets/SConcertPropertyChainChip.h b/Engine/Plugins/Developer/Concert/ConcertScripting/ConcertReplicationScripting/Source/ConcertReplicationScriptingEditor/Private/Widgets/SConcertPropertyChainChip.h
new file mode 100644
--- /dev/null
+++ b/Engine/Plugins/Developer/Concert/ConcertScripting/ConcertReplicationScripting/Source/ConcertReplicationScriptingEditor/Private/Widgets/SConcertPropertyChainChip.h
@@ -0,0 +1,35 @@
+// Copyright Epic Games, Inc. All Rights Reserved.
+
+#pragma once
+
+#include "PropertyChainDelegates.h"
+#include "Replication/Data/ConcertPropertySelection.h"
+
+#include "Widgets/DeclarativeSyntaxSupport.h"
+#include "Widgets/SCompoundWidget.h"
+
+namespace UE::ConcertReplicationScriptingEditor
+{
+	/** Displays a FConcertPropertyChain */
+	class SConcertPropertyChainChip : public SCompoundWidget
+	{
+	public:
+		
+		SLATE_BEGIN_ARGS(SConcertPropertyChainChip)
+		{}
+			/** The property to display */
+			SLATE_ATTRIBUTE(FConcertPropertyChain, DisplayedProperty)
+
+			/** Whether the X-Button should be shown */
+			SLATE_ARGUMENT(bool, ShowClearButton)
+
+			/** Called when the X-button is pressed */
+			SLATE_EVENT(FOnPressed, OnClearPressed)
+
+			/** Called when the chip is pressed. */
+			SLATE_EVENT(FOnPressed, OnEditPressed)
+		SLATE_END_ARGS()
+
+		void Construct(const FArguments& InArgs);
+	};
+}
diff --git a/Engine/Plugins/Developer/Concert/ConcertScripting/ConcertReplicationScripting/Source/ConcertReplicationScriptingEditor/Private/Widgets/SConcertPropertyChainCombo.cpp b/Engine/Plugins/Developer/Concert/ConcertScripting/ConcertReplicationScripting/Source/ConcertReplicationScriptingEditor/Private/Widgets/SConcertPropertyChainCombo.cpp
new file mode 100644
--- /dev/null
+++ b/Engine/Plugins/Developer/Concert/ConcertScripting/ConcertReplicationScripting/Source/ConcertReplicationScriptingEditor/Private/Widgets/SConcertPropertyChainCombo.cpp
@@ -0,0 +1,122 @@
+// Copyright Epic Games, Inc. All Rights Reserved.
+
+#include "SConcertPropertyChainCombo.h"
+
+#include "ConcertPropertyChainWrapper.h"
+#include "SConcertPropertyChainChip.h"
+#include "SConcertPropertyChainPicker.h"
+
+#include "Styling/AppStyle.h"
+#include "Widgets/Input/SComboButton.h"
+#include "Widgets/Layout/SWidgetSwitcher.h"
+
+#define LOCTEXT_NAMESPACE "SConcertPropertyChainCombo"
+
+namespace UE::ConcertReplicationScriptingEditor
+{
+	void SConcertPropertyChainCombo::Construct(const FArguments& InArgs)
+	{
+		LastClassUsed = InArgs._InitialClassSelection;
+		ContainedProperties = InArgs._ContainedProperties;
+		bIsEditable = InArgs._IsEditable;
+		HasMultipleValuesAttribute = InArgs._HasMultipleValues;
+		OnClassChangedDelegate = InArgs._OnClassChanged;
+		OnPropertySelectionChangedDelegate = InArgs._OnPropertySelectionChanged;
+		
+		ChildSlot
+		[
+			SAssignNew(ComboButton, SComboButton)
+			.HasDownArrow(true)
+			.ContentPadding(1)
+			.OnGetMenuContent(this, &SConcertPropertyChainCombo::OnGetMenuContent)
+			.ButtonContent()
+			[
+				SNew(SWidgetSwitcher)
+				.WidgetIndex_Lambda([this](){ return HasMultipleValuesAttribute.Get() ? 0 : 1; })
+				
+				+SWidgetSwitcher::Slot()
+				[
+					SNew(STextBlock)
+					.Text(LOCTEXT("MultipleValues", "Multiple Values"))
+				]
+				+SWidgetSwitcher::Slot()
+				[
+					SAssignNew(PropertyListView, SListView<TSharedPtr<FConcertPropertyChain>>)
+					.ListItemsSource(&DisplayedProperties)
+					.SelectionMode(ESelectionMode::None)
+					.ItemHeight(23.0f)
+					.ListViewStyle(&FAppStyle::Get().GetWidgetStyle<FTableViewStyle>("SimpleListView"))
+					.OnGenerateRow(this, &SConcertPropertyChainCombo::MakePropertyChainRow)
+				]
+			]
+		];
+
+		RefreshPropertyContent();
+	}
+
+	void SConcertPropertyChainCombo::RefreshPropertyContent()
+	{
+		DisplayedProperties.Reset();
+		ON_SCOPE_EXIT{ PropertyListView->RequestListRefresh(); };
+		
+		if (HasMultipleValuesAttribute.Get() || !ensureAlwaysMsgf(ContainedProperties, TEXT("You forgot to specified ContainedProperties during construction!")))
+		{
+			return;
+		}
+
+		for (const FConcertPropertyChain& Property : *ContainedProperties)
+		{
+			DisplayedProperties.Emplace(MakeShared<FConcertPropertyChain>(Property));
+		}
+		DisplayedProperties.Sort([](const TSharedPtr<FConcertPropertyChain>& Left, const TSharedPtr<FConcertPropertyChain>& Right)
+		{
+			return Left->ToString() <= Right->ToString();
+		});
+	}
+
+	TSharedRef<SWidget> SConcertPropertyChainCombo::OnGetMenuContent()
+	{
+		 const TSharedRef<SConcertPropertyChainPicker> Menu = SNew(SConcertPropertyChainPicker)
+			.ContainedProperties(ContainedProperties)
+			.IsEditable(bIsEditable)
+			.InitialClassSelection(LastClassUsed)
+			.OnClassChanged_Lambda([this](const UClass* Class)
+			{
+				LastClassUsed = Class;
+				OnClassChangedDelegate.ExecuteIfBound(Class);
+			})
+			.OnSelectedPropertiesChanged(OnPropertySelectionChangedDelegate);
+		WeakButtonMenu = Menu;
+		return Menu;
+	}
+
+	TSharedRef<ITableRow> SConcertPropertyChainCombo::MakePropertyChainRow(TSharedPtr<FConcertPropertyChain> Item, const TSharedRef<STableViewBase>& OwnerTable)
+	{
+		const FConcertPropertyChain PropertyChain = *Item.Get();
+		return SNew(STableRow<TSharedPtr<FConcertPropertyChain>>, OwnerTable)
+			.Style(&FAppStyle::Get().GetWidgetStyle<FTableRowStyle>("SimpleTableView.Row"))
+			.Padding(FMargin(0,2))
+			[
+				SNew(SConcertPropertyChainChip)
+				.DisplayedProperty(*Item)
+				.ShowClearButton(bIsEditable)
+				.OnClearPressed_Lambda([this, PropertyChain](){ OnPropertySelectionChangedDelegate.ExecuteIfBound(PropertyChain, false); })
+				.OnEditPressed_Lambda([this, PropertyChain]()
+				{
+					// Handle being pressed while menu is open already
+					if (!ComboButton->IsOpen())
+					{
+						ComboButton->SetIsOpen(true);
+					}
+					
+					// Should have been created by SetIsOpen 
+					if (const TSharedPtr<SConcertPropertyChainPicker> MenuPin = WeakButtonMenu.Pin())
+					{
+						MenuPin->RequestScrollIntoView(PropertyChain);
+					}
+				})
+			];
+	}
+}
+
+#undef LOCTEXT_NAMESPACE
\ No newline at end of file
diff --git a/Engine/Plugins/Developer/Concert/ConcertScripting/ConcertReplicationScripting/Source/ConcertReplicationScriptingEditor/Private/Widgets/SConcertPropertyChainCombo.h b/Engine/Plugins/Developer/Concert/ConcertScripting/ConcertReplicationScripting/Source/ConcertReplicationScriptingEditor/Private/Widgets/SConcertPropertyChainCombo.h
new file mode 100644
--- /dev/null
+++ b/Engine/Plugins/Developer/Concert/ConcertScripting/ConcertReplicationScripting/Source/ConcertReplicationScriptingEditor/Private/Widgets/SConcertPropertyChainCombo.h
@@ -0,0 +1,84 @@
+// Copyright Epic Games, Inc. All Rights Reserved.
+
+#pragma once
+
+#include "Replication/Data/ConcertPropertySelection.h"
+#include "PropertyChainDelegates.h"
+
+#include "Widgets/DeclarativeSyntaxSupport.h"
+#include "Widgets/SCompoundWidget.h"
+#include "Widgets/Views/SListView.h"
+
+class SComboButton;
+
+namespace UE::ConcertReplicationScriptingEditor
+{
+	class SConcertPropertyChainPicker;
+	/**
+	 * A combo button that displays a list of properties as content and as menu opens a property tree view for editing
+	 * the list of properties.
+	 */
+	class SConcertPropertyChainCombo : public SCompoundWidget
+	{
+	public:
+
+		SLATE_BEGIN_ARGS(SConcertPropertyChainCombo)
+			: _InitialClassSelection(nullptr)
+		{}
+			/** The properties that are supposed to be displayed inside of the button */
+			SLATE_ARGUMENT(const TSet<FConcertPropertyChain>*, ContainedProperties)
+
+			/** Whether the UI should allow editing */
+			SLATE_ARGUMENT(bool, IsEditable)
+			/** If this returns true, then the button will display "Multiple Values" instead. */
+			SLATE_ATTRIBUTE(bool, HasMultipleValues)
+		
+			/** The initial class to search in */
+			SLATE_ARGUMENT(const UClass*, InitialClassSelection)
+		
+			/** Called when the class is changed */
+			SLATE_EVENT(FOnClassChanged, OnClassChanged)
+			/** Called when the selection of a property changes. */
+			SLATE_EVENT(FOnSelectedPropertiesChanged, OnPropertySelectionChanged)
+		SLATE_END_ARGS()
+
+		void Construct(const FArguments& InArgs);
+
+		/** Regenerates the button's content. */
+		void RefreshPropertyContent();
+		
+	private:
+
+		/** Class set by the last menu this widget had used. */
+		const UClass* LastClassUsed = nullptr;
+		/** The properties that are supposed to be displayed inside of the button */
+		const TSet<FConcertPropertyChain>* ContainedProperties;
+
+		/** Whether the UI should allow editing. */
+		bool bIsEditable;
+
+		/** The combo button being displayed. */
+		TSharedPtr<SComboButton> ComboButton;
+		/** Used to scroll into view when clicking an item in the button content. */
+		TWeakPtr<SConcertPropertyChainPicker> WeakButtonMenu;
+		
+		/** The properties displayed in PropertyListView. */
+		TArray<TSharedPtr<FConcertPropertyChain>> DisplayedProperties;
+		/** The content of this button */
+		TSharedPtr<SListView<TSharedPtr<FConcertPropertyChain>>> PropertyListView;
+
+		/** If true, then ContainedProperties contains the accumulated values of multiple sources. If false, all sources have the same value. */
+		TAttribute<bool> HasMultipleValuesAttribute;
+		
+		/** Triggered when the class changes */
+		FOnClassChanged OnClassChangedDelegate;
+		/** Called when the selection of a property changes. */
+		FOnSelectedPropertiesChanged OnPropertySelectionChangedDelegate;
+
+		/** Generates the combo button menu content. */
+		TSharedRef<SWidget> OnGetMenuContent();
+		
+		TSharedRef<ITableRow> MakePropertyChainRow(TSharedPtr<FConcertPropertyChain> Item, const TSharedRef<STableViewBase>& OwnerTable);
+	};
+}
+
diff --git a/Engine/Plugins/Developer/Concert/ConcertScripting/ConcertReplicationScripting/Source/ConcertReplicationScriptingEditor/Private/Widgets/SConcertPropertyChainPicker.cpp b/Engine/Plugins/Developer/Concert/ConcertScripting/ConcertReplicationScripting/Source/ConcertReplicationScriptingEditor/Private/Widgets/SConcertPropertyChainPicker.cpp
new file mode 100644
--- /dev/null
+++ b/Engine/Plugins/Developer/Concert/ConcertScripting/ConcertReplicationScripting/Source/ConcertReplicationScriptingEditor/Private/Widgets/SConcertPropertyChainPicker.cpp
@@ -0,0 +1,145 @@
+// Copyright Epic Games, Inc. All Rights Reserved.
+
+#include "SConcertPropertyChainPicker.h"
+
+#include "PropertyCustomizationHelpers.h"
+#include "PropertySelectionColumn.h"
+#include "Replication/PropertyChainUtils.h"
+#include "Replication/ReplicationWidgetFactories.h"
+#include "Replication/Editor/View/IPropertyTreeView.h"
+#include "Widgets/Layout/SBorder.h"
+#include "Widgets/Layout/SBox.h"
+#include "Widgets/SBoxPanel.h"
+#include "Widgets/Text/STextBlock.h"
+
+#define LOCTEXT_NAMESPACE "SConcertPropertyChainPicker"
+
+namespace UE::ConcertReplicationScriptingEditor
+{
+	void SConcertPropertyChainPicker::Construct(const FArguments& InArgs)
+	{
+		SelectedClass = InArgs._InitialClassSelection;
+		ContainedProperties = InArgs._ContainedProperties;
+		OnClassChangedDelegate = InArgs._OnClassChanged;
+		OnSelectedPropertiesChangedDelegate = InArgs._OnSelectedPropertiesChanged;
+		
+		ChildSlot
+		[
+			SNew(SBox)
+			.MinDesiredHeight(750.f)
+			.MaxDesiredHeight(750.f)
+			.MinDesiredWidth(350.f)
+			.MaxDesiredWidth(350.f)
+			[
+				SNew(SBorder)
+				.Padding(2.f)
+				.BorderImage(FStyleDefaults::GetNoBrush())
+				[
+					SNew(SVerticalBox)
+
+					+SVerticalBox::Slot()
+					.AutoHeight()
+					.Padding(2.f)
+					[
+						CreateClassPickerSection()
+					]
+
+					+SVerticalBox::Slot()
+					.FillHeight(1.f)
+					.Padding(2.f)
+					[
+						CreatePropertyTreeView(InArgs)
+					]
+				]
+			]
+		];
+		
+		RefreshPropertiesDisplayedInTree();
+	}
+
+	void SConcertPropertyChainPicker::RequestScrollIntoView(const FConcertPropertyChain& PropertyChain)
+	{
+		TreeView->RequestScrollIntoView(PropertyChain);
+	}
+
+	TSharedRef<SWidget> SConcertPropertyChainPicker::CreateClassPickerSection()
+	{
+		return SNew(SHorizontalBox)
+			.ToolTipText(LOCTEXT("Class.Tooltip", "The class for which to display the tree hierarchy"))
+				
+			+SHorizontalBox::Slot()
+			.AutoWidth()
+			.HAlign(HAlign_Center)
+			.VAlign(VAlign_Center)
+			[
+				SNew(STextBlock)
+				.Text(LOCTEXT("Class.Label", "Class Context"))
+			]
+				
+			+SHorizontalBox::Slot()
+			.HAlign(HAlign_Fill)
+			.VAlign(VAlign_Center)
+			.Padding(50.f, 0.f, 0.f, 0.f)
+			[
+				SNew(SClassPropertyEntryBox)
+				.MetaClass(UObject::StaticClass())
+				.SelectedClass_Lambda([this]() { return SelectedClass; })
+				.OnSetClass(this, &SConcertPropertyChainPicker::OnSetClass)
+			];
+	}
+
+	TSharedRef<SWidget> SConcertPropertyChainPicker::CreatePropertyTreeView(const FArguments& InArgs)
+	{
+		ConcertSharedSlate::FCreatePropertyTreeViewParams Params
+		{
+			.PropertyColumns =
+			{
+				ConcertSharedSlate::ReplicationColumns::Property::LabelColumn()
+			},
+		};
+		Params.NoItemsContent.Widget = SNew(STextBlock).Text(LOCTEXT("NoClass", "Select a class"));
+
+		if (ensureAlwaysMsgf(ContainedProperties, TEXT("You forgot to pass in the required ContainedProperties argument for widget construction!")))
+		{
+			Params.PrimaryPropertySort = { PropertySelectionCheckboxColumnId, EColumnSortMode::Ascending };
+			Params.PropertyColumns.Add(
+				MakePropertySelectionCheckboxColumn(*ContainedProperties, FOnSelectProperty::CreateSP(this, &SConcertPropertyChainPicker::OnPropertySelected), InArgs._IsEditable)
+				);
+		}
+		
+		TreeView = ConcertSharedSlate::CreateSearchablePropertyTreeView(Params);
+		return TreeView->GetWidget();
+	}
+	
+	void SConcertPropertyChainPicker::OnSetClass(const UClass* Class)
+	{
+		SelectedClass = Class;
+		OnClassChangedDelegate.ExecuteIfBound(Class);
+		RefreshPropertiesDisplayedInTree();
+	}
+
+	void SConcertPropertyChainPicker::RefreshPropertiesDisplayedInTree()
+	{
+		TSet<FConcertPropertyChain> Properties;
+		FSoftClassPath ClassPath = SelectedClass;
+		
+		if (SelectedClass)
+		{
+			ConcertSyncCore::PropertyChain::ForEachReplicatableConcertProperty(*SelectedClass, [&Properties](FConcertPropertyChain&& PropertyChain)
+			{
+				Properties.Emplace(MoveTemp(PropertyChain));
+				return EBreakBehavior::Continue;
+			});
+		}
+
+		// If not class is selected, the widget will display the NoItemsContent ("Select a class").
+		TreeView->RefreshPropertyData(Properties, ClassPath);
+	}
+
+	void SConcertPropertyChainPicker::OnPropertySelected(const FConcertPropertyChain& ConcertPropertyChain, bool bIsSelected)
+	{
+		OnSelectedPropertiesChangedDelegate.ExecuteIfBound(ConcertPropertyChain, bIsSelected);
+	}
+}
+
+#undef LOCTEXT_NAMESPACE
\ No newline at end of file
diff --git a/Engine/Plugins/Developer/Concert/ConcertScripting/ConcertReplicationScripting/Source/ConcertReplicationScriptingEditor/Private/Widgets/SConcertPropertyChainPicker.h b/Engine/Plugins/Developer/Concert/ConcertScripting/ConcertReplicationScripting/Source/ConcertReplicationScriptingEditor/Private/Widgets/SConcertPropertyChainPicker.h
new file mode 100644
--- /dev/null
+++ b/Engine/Plugins/Developer/Concert/ConcertScripting/ConcertReplicationScripting/Source/ConcertReplicationScriptingEditor/Private/Widgets/SConcertPropertyChainPicker.h
@@ -0,0 +1,69 @@
+// Copyright Epic Games, Inc. All Rights Reserved.
+
+#pragma once
+
+#include "Replication/Data/ConcertPropertySelection.h"
+#include "PropertyChainDelegates.h"
+
+#include "Widgets/DeclarativeSyntaxSupport.h"
+#include "Widgets/SCompoundWidget.h"
+
+namespace UE::ConcertSharedSlate
+{
+	class IPropertyTreeView;
+}
+
+namespace UE::ConcertReplicationScriptingEditor
+{
+	
+	/** Menu that has a class picker and tree view of properties in the class. */
+	class SConcertPropertyChainPicker : public SCompoundWidget
+	{
+	public:
+		
+		SLATE_BEGIN_ARGS(SConcertPropertyChainPicker)
+			: _InitialClassSelection(nullptr)
+			, _ContainedProperties(nullptr)
+		{}
+			/** The initial class to search in */
+			SLATE_ARGUMENT(const UClass*, InitialClassSelection)
+			/** Whether the UI should allow editing */
+			SLATE_ARGUMENT(bool, IsEditable)
+		
+			/** The properties that are supposed to be displayed inside of the button */
+			SLATE_ARGUMENT(const TSet<FConcertPropertyChain>*, ContainedProperties)
+		
+			/** Called when the class is changed */
+			SLATE_EVENT(FOnClassChanged, OnClassChanged)
+			/** Called when the user's property selection changes */
+			SLATE_EVENT(FOnSelectedPropertiesChanged, OnSelectedPropertiesChanged)
+		SLATE_END_ARGS()
+
+		void Construct(const FArguments& InArgs);
+
+		void RequestScrollIntoView(const FConcertPropertyChain& PropertyChain);
+			
+	private:
+		
+		/** The class for which to display the hierarchy*/
+		const UClass* SelectedClass = nullptr;
+		/** The properties that are supposed to be displayed inside of the button */
+		const TSet<FConcertPropertyChain>* ContainedProperties;
+
+		/** Displays the properties */
+		TSharedPtr<ConcertSharedSlate::IPropertyTreeView> TreeView;
+		
+		/** Triggered when the class changes */
+		FOnClassChanged OnClassChangedDelegate;
+		/** Called when the user's property selection changes */
+		FOnSelectedPropertiesChanged OnSelectedPropertiesChangedDelegate;
+
+		TSharedRef<SWidget> CreateClassPickerSection();
+		TSharedRef<SWidget> CreatePropertyTreeView(const FArguments& InArgs);
+		
+		void OnSetClass(const UClass* Class);
+		void RefreshPropertiesDisplayedInTree();
+		void OnPropertySelected(const FConcertPropertyChain& ConcertPropertyChain, bool bIsSelected);
+	};
+}
+
diff --git a/Engine/Plugins/Developer/Concert/ConcertUI/ConcertClientSharedSlate/Source/ConcertClientSharedSlate/Private/Replication/Editor/View/PropertyTree/SFilteredPropertyTreeView.h b/Engine/Plugins/Developer/Concert/ConcertUI/ConcertClientSharedSlate/Source/ConcertClientSharedSlate/Private/Replication/Editor/View/PropertyTree/SFilteredPropertyTreeView.h
--- a/Engine/Plugins/Developer/Concert/ConcertUI/ConcertClientSharedSlate/Source/ConcertClientSharedSlate/Private/Replication/Editor/View/PropertyTree/SFilteredPropertyTreeView.h
+++ b/Engine/Plugins/Developer/Concert/ConcertUI/ConcertClientSharedSlate/Source/ConcertClientSharedSlate/Private/Replication/Editor/View/PropertyTree/SFilteredPropertyTreeView.h
@@ -58,6 +58,7 @@ namespace UE::ConcertClientSharedSlate
 		}
 		virtual void RequestRefilter() const override { ExtendedTreeView->RequestRefilter(); }
 		virtual void RequestResortForColumn(const FName& ColumnId) override { ExtendedTreeView->RequestResortForColumn(ColumnId); }
+		virtual void RequestScrollIntoView(const FConcertPropertyChain& PropertyChain) override { ExtendedTreeView->RequestScrollIntoView(PropertyChain); }
 		virtual TSharedRef<SWidget> GetWidget() override { return SharedThis(this); }
 		//~ End IPropertyTreeView Interface
 
diff --git a/Engine/Plugins/Developer/Concert/ConcertUI/ConcertSharedSlate/Source/ConcertSharedSlate/Private/Replication/Editor/View/ObjectViewer/Property/SPropertyTreeView.cpp b/Engine/Plugins/Developer/Concert/ConcertUI/ConcertSharedSlate/Source/ConcertSharedSlate/Private/Replication/Editor/View/ObjectViewer/Property/SPropertyTreeView.cpp
--- a/Engine/Plugins/Developer/Concert/ConcertUI/ConcertSharedSlate/Source/ConcertSharedSlate/Private/Replication/Editor/View/ObjectViewer/Property/SPropertyTreeView.cpp
+++ b/Engine/Plugins/Developer/Concert/ConcertUI/ConcertSharedSlate/Source/ConcertSharedSlate/Private/Replication/Editor/View/ObjectViewer/Property/SPropertyTreeView.cpp
@@ -59,6 +59,19 @@ namespace UE::ConcertSharedSlate
 		TreeView->RequestRefilter();
 	}
 
+	void SPropertyTreeView::RequestScrollIntoView(const FConcertPropertyChain& PropertyChain)
+	{
+		const int32 Index = PropertyRowData.IndexOfByPredicate([&PropertyChain](const TSharedPtr<FReplicatedPropertyData>& Data)
+		{
+			return Data->GetProperty() == PropertyChain;
+		});
+		if (PropertyRowData.IsValidIndex(Index))
+		{
+			TreeView->SetExpandedItems({ PropertyRowData[Index] }, true);
+			TreeView->RequestScrollIntoView(PropertyRowData[Index]);
+		}
+	}
+
 	TSharedRef<FReplicatedPropertyData> SPropertyTreeView::AllocatePropertyData(FSoftClassPath OwningClass, FConcertPropertyChain PropertyChain)
 	{
 		return MakeShared<FReplicatedPropertyData>(MoveTemp(OwningClass), MoveTemp(PropertyChain));
diff --git a/Engine/Plugins/Developer/Concert/ConcertUI/ConcertSharedSlate/Source/ConcertSharedSlate/Private/Replication/Editor/View/ObjectViewer/Property/SPropertyTreeView.h b/Engine/Plugins/Developer/Concert/ConcertUI/ConcertSharedSlate/Source/ConcertSharedSlate/Private/Replication/Editor/View/ObjectViewer/Property/SPropertyTreeView.h
--- a/Engine/Plugins/Developer/Concert/ConcertUI/ConcertSharedSlate/Source/ConcertSharedSlate/Private/Replication/Editor/View/ObjectViewer/Property/SPropertyTreeView.h
+++ b/Engine/Plugins/Developer/Concert/ConcertUI/ConcertSharedSlate/Source/ConcertSharedSlate/Private/Replication/Editor/View/ObjectViewer/Property/SPropertyTreeView.h
@@ -64,6 +64,7 @@ namespace UE::ConcertSharedSlate
 		virtual void RefreshPropertyData(const TSet<FConcertPropertyChain>& PropertiesToDisplay, const FSoftClassPath& Class, bool bCanReuseExistingRowItems) override;
 		virtual void RequestRefilter() const override { TreeView->RequestRefilter(); }
 		virtual void RequestResortForColumn(const FName& ColumnId) override { TreeView->RequestResortForColumn(ColumnId); }
+		virtual void RequestScrollIntoView(const FConcertPropertyChain& PropertyChain) override;
 		virtual TSharedRef<SWidget> GetWidget() override { return SharedThis(this); }
 		//~ Begin IPropertyTreeView Interface
 
diff --git a/Engine/Plugins/Developer/Concert/ConcertUI/ConcertSharedSlate/Source/ConcertSharedSlate/Private/Replication/Editor/View/Tree/SReplicationTreeView.h b/Engine/Plugins/Developer/Concert/ConcertUI/ConcertSharedSlate/Source/ConcertSharedSlate/Private/Replication/Editor/View/Tree/SReplicationTreeView.h
--- a/Engine/Plugins/Developer/Concert/ConcertUI/ConcertSharedSlate/Source/ConcertSharedSlate/Private/Replication/Editor/View/Tree/SReplicationTreeView.h
+++ b/Engine/Plugins/Developer/Concert/ConcertUI/ConcertSharedSlate/Source/ConcertSharedSlate/Private/Replication/Editor/View/Tree/SReplicationTreeView.h
@@ -263,6 +263,11 @@ namespace UE::ConcertSharedSlate
 				TreeView->SetItemExpansion(Item, bIsExpanded);
 			}
 		}
+
+		void RequestScrollIntoView(const TSharedPtr<TItemType>& Item)
+		{
+			TreeView->RequestScrollIntoView(Item);
+		}
 		
 		TArray<TSharedPtr<TItemType>> GetSelectedItems() const { return TreeView->GetSelectedItems(); }
 		const TArray<TSharedPtr<TItemType>>& GetFilteredRootItems() const { return FilteredRootItems; }
diff --git a/Engine/Plugins/Developer/Concert/ConcertUI/ConcertSharedSlate/Source/ConcertSharedSlate/Public/Replication/Editor/View/IPropertyTreeView.h b/Engine/Plugins/Developer/Concert/ConcertUI/ConcertSharedSlate/Source/ConcertSharedSlate/Public/Replication/Editor/View/IPropertyTreeView.h
--- a/Engine/Plugins/Developer/Concert/ConcertUI/ConcertSharedSlate/Source/ConcertSharedSlate/Public/Replication/Editor/View/IPropertyTreeView.h
+++ b/Engine/Plugins/Developer/Concert/ConcertUI/ConcertSharedSlate/Source/ConcertSharedSlate/Public/Replication/Editor/View/IPropertyTreeView.h
@@ -26,7 +26,7 @@ namespace UE::ConcertSharedSlate
 		 *	Set this to false, if all rows should be regenerated (clears selection).
 		 *	In general, always set this to false if you've changed the object for which you're displaying the class.
 		 */
-		virtual void RefreshPropertyData(const TSet<FConcertPropertyChain>& PropertiesToDisplay, const FSoftClassPath& Class, bool bCanReuseExistingRowItems) = 0;
+		virtual void RefreshPropertyData(const TSet<FConcertPropertyChain>& PropertiesToDisplay, const FSoftClassPath& Class, bool bCanReuseExistingRowItems = true) = 0;
 		
 		/**
 		 * Reapply the filter function to all items at the end of the frame. Call e.g. when the filters have changed.
@@ -39,6 +39,9 @@ namespace UE::ConcertSharedSlate
 		 */
 		virtual void RequestResortForColumn(const FName& ColumnId) = 0;
 
+		/** Scroll the given property into view, if it is contained. */
+		virtual void RequestScrollIntoView(const FConcertPropertyChain& PropertyChain) = 0;
+
 		/** Gets the tree view's widget */
 		virtual TSharedRef<SWidget> GetWidget() = 0;
 
