Add missing dependency from ConcertReplicationScripting to ConcertSharedSlate

diff --git a/Engine/Plugins/Developer/Concert/ConcertScripting/ConcertReplicationScripting/ConcertReplicationScripting.uplugin b/Engine/Plugins/Developer/Concert/ConcertScripting/ConcertReplicationScripting/ConcertReplicationScripting.uplugin
--- a/Engine/Plugins/Developer/Concert/ConcertScripting/ConcertReplicationScripting/ConcertReplicationScripting.uplugin
+++ b/Engine/Plugins/Developer/Concert/ConcertScripting/ConcertReplicationScripting/ConcertReplicationScripting.uplugin
@@ -26,6 +26,10 @@
         }
 	],
 	"Plugins": [
+		{
+			"Name": "ConcertSharedSlate",
+			"Enabled": true
+		},
 		{
 			"Name": "ConcertSyncCore",
 			"Enabled": true
