Add FConcertPropertyChainWrapperContainer, a special array version with a detail customization for selecting multiple properties

diff --git a/Engine/Plugins/Developer/Concert/ConcertScripting/ConcertReplicationScripting/Source/ConcertReplicationScripting/Private/ConcertPropertyChainWrapper.cpp b/Engine/Plugins/Developer/Concert/ConcertScripting/ConcertReplicationScripting/Source/ConcertReplicationScripting/Private/ConcertPropertyChainWrapper.cpp
new file mode 100644
--- /dev/null
+++ b/Engine/Plugins/Developer/Concert/ConcertScripting/ConcertReplicationScripting/Source/ConcertReplicationScripting/Private/ConcertPropertyChainWrapper.cpp
@@ -0,0 +1,8 @@
+// Copyright Epic Games, Inc. All Rights Reserved.
+
+#include "ConcertPropertyChainWrapper.h"
+
+uint32 GetTypeHash(const FConcertPropertyChainWrapper& ChainWrapper)
+{
+	return GetTypeHash(ChainWrapper.PropertyChain);
+}
\ No newline at end of file
diff --git a/Engine/Plugins/Developer/Concert/ConcertScripting/ConcertReplicationScripting/Source/ConcertReplicationScripting/Public/ConcertPropertyChainWrapper.h b/Engine/Plugins/Developer/Concert/ConcertScripting/ConcertReplicationScripting/Source/ConcertReplicationScripting/Public/ConcertPropertyChainWrapper.h
--- a/Engine/Plugins/Developer/Concert/ConcertScripting/ConcertReplicationScripting/Source/ConcertReplicationScripting/Public/ConcertPropertyChainWrapper.h
+++ b/Engine/Plugins/Developer/Concert/ConcertScripting/ConcertReplicationScripting/Source/ConcertReplicationScripting/Public/ConcertPropertyChainWrapper.h
@@ -17,4 +17,9 @@ struct CONCERTREPLICATIONSCRIPTING_API FConcertPropertyChainWrapper
 	explicit FConcertPropertyChainWrapper(FConcertPropertyChain InPropertyChain)
 		: PropertyChain(MoveTemp(InPropertyChain))
 	{}
-};
\ No newline at end of file
+
+	friend bool operator==(const FConcertPropertyChainWrapper& Left, const FConcertPropertyChainWrapper& Right) { return Left.PropertyChain == Right.PropertyChain; }
+	friend bool operator!=(const FConcertPropertyChainWrapper& Left, const FConcertPropertyChainWrapper& Right) { return !(Left == Right); }
+};
+
+CONCERTREPLICATIONSCRIPTING_API uint32 GetTypeHash(const FConcertPropertyChainWrapper& ChainWrapper);
\ No newline at end of file
diff --git a/Engine/Plugins/Developer/Concert/ConcertScripting/ConcertReplicationScripting/Source/ConcertReplicationScripting/Public/ConcertPropertyChainWrapperContainer.h b/Engine/Plugins/Developer/Concert/ConcertScripting/ConcertReplicationScripting/Source/ConcertReplicationScripting/Public/ConcertPropertyChainWrapperContainer.h
new file mode 100644
--- /dev/null
+++ b/Engine/Plugins/Developer/Concert/ConcertScripting/ConcertReplicationScripting/Source/ConcertReplicationScripting/Public/ConcertPropertyChainWrapperContainer.h
@@ -0,0 +1,16 @@
+// Copyright Epic Games, Inc. All Rights Reserved.
+
+#pragma once
+
+#include "Replication/Data/ConcertPropertySelection.h"
+#include "ConcertPropertyChainWrapperContainer.generated.h"
+
+/** Special array of property chains. Used to allow for detail customization. */
+USTRUCT(BlueprintType, meta = (DisplayName = "Concert Property Chain Container"))
+struct CONCERTREPLICATIONSCRIPTING_API FConcertPropertyChainWrapperContainer
+{
+	GENERATED_BODY()
+
+	UPROPERTY(BlueprintReadWrite, EditAnywhere, Category = "Concert|Replication")
+	TArray<FConcertPropertyChainWrapper> PropertyChains;
+};
\ No newline at end of file
diff --git a/Engine/Plugins/Developer/Concert/ConcertScripting/ConcertReplicationScripting/Source/ConcertReplicationScriptingEditor/Private/ConcertReplicationScriptingEditorModule.cpp b/Engine/Plugins/Developer/Concert/ConcertScripting/ConcertReplicationScripting/Source/ConcertReplicationScriptingEditor/Private/ConcertReplicationScriptingEditorModule.cpp
--- a/Engine/Plugins/Developer/Concert/ConcertScripting/ConcertReplicationScripting/Source/ConcertReplicationScriptingEditor/Private/ConcertReplicationScriptingEditorModule.cpp
+++ b/Engine/Plugins/Developer/Concert/ConcertScripting/ConcertReplicationScripting/Source/ConcertReplicationScriptingEditor/Private/ConcertReplicationScriptingEditorModule.cpp
@@ -3,8 +3,11 @@
 #include "ConcertReplicationScriptingEditorModule.h"
 
 #include "ConcertPropertyChainWrapper.h"
-#include "ReplicationScriptingStyle.h"
+#include "ConcertPropertyChainWrapperContainer.h"
+#include "Customization/ConcertPropertyContainerCustomization.h"
 #include "Customization/ConcertPropertyCustomization.h"
+#include "ReplicationScriptingStyle.h"
+
 #include "Modules/ModuleManager.h"
 #include "PropertyEditorModule.h"
 
@@ -17,7 +20,11 @@ namespace UE::ConcertReplicationScriptingEditor
 		FPropertyEditorModule& PropertyModule = FModuleManager::LoadModuleChecked<FPropertyEditorModule>("PropertyEditor");
 		PropertyModule.RegisterCustomPropertyTypeLayout(
 			FConcertPropertyChainWrapper::StaticStruct()->GetFName(),
-			FOnGetPropertyTypeCustomizationInstance::CreateStatic(&FConcertPropertyCustomization::MakeInstance)
+			FOnGetPropertyTypeCustomizationInstance::CreateStatic(&FConcertPropertyCustomization::MakeInstance, &SharedClassRememberer)
+			);
+		PropertyModule.RegisterCustomPropertyTypeLayout(
+			FConcertPropertyChainWrapperContainer::StaticStruct()->GetFName(),
+			FOnGetPropertyTypeCustomizationInstance::CreateStatic(&FConcertPropertyContainerCustomization::MakeInstance, &SharedClassRememberer)
 			);
 	}
 
diff --git a/Engine/Plugins/Developer/Concert/ConcertScripting/ConcertReplicationScripting/Source/ConcertReplicationScriptingEditor/Private/ConcertReplicationScriptingEditorModule.h b/Engine/Plugins/Developer/Concert/ConcertScripting/ConcertReplicationScripting/Source/ConcertReplicationScriptingEditor/Private/ConcertReplicationScriptingEditorModule.h
--- a/Engine/Plugins/Developer/Concert/ConcertScripting/ConcertReplicationScripting/Source/ConcertReplicationScriptingEditor/Private/ConcertReplicationScriptingEditorModule.h
+++ b/Engine/Plugins/Developer/Concert/ConcertScripting/ConcertReplicationScripting/Source/ConcertReplicationScriptingEditor/Private/ConcertReplicationScriptingEditorModule.h
@@ -3,6 +3,7 @@
 #pragma once
 
 #include "Modules/ModuleInterface.h"
+#include "Util/ClassRememberer.h"
 
 namespace UE::ConcertReplicationScriptingEditor
 {
@@ -14,5 +15,10 @@ namespace UE::ConcertReplicationScriptingEditor
 		virtual void StartupModule() override;
 		virtual void ShutdownModule() override;
 		//~ End IModuleInterface Interface
+
+	private:
+
+		/** Customizations use this cache so users do not constantly have to re-select the class in the drop-down menus. */
+		FClassRememberer SharedClassRememberer;
 	};
-}
\ No newline at end of file
+}
diff --git a/Engine/Plugins/Developer/Concert/ConcertScripting/ConcertReplicationScripting/Source/ConcertReplicationScriptingEditor/Private/Customization/ConcertPropertyContainerCustomization.cpp b/Engine/Plugins/Developer/Concert/ConcertScripting/ConcertReplicationScripting/Source/ConcertReplicationScriptingEditor/Private/Customization/ConcertPropertyContainerCustomization.cpp
new file mode 100644
--- /dev/null
+++ b/Engine/Plugins/Developer/Concert/ConcertScripting/ConcertReplicationScripting/Source/ConcertReplicationScriptingEditor/Private/Customization/ConcertPropertyContainerCustomization.cpp
@@ -0,0 +1,139 @@
+// Copyright Epic Games, Inc. All Rights Reserved.
+
+#include "ConcertPropertyContainerCustomization.h"
+
+#include "ConcertPropertyChainWrapper.h"
+#include "ConcertPropertyChainWrapperContainer.h"
+#include "Util/ClassRememberer.h"
+#include "Widgets/SConcertPropertyChainCombo.h"
+
+#include "DetailWidgetRow.h"
+#include "ScopedTransaction.h"
+
+#define LOCTEXT_NAMESPACE "FConcertPropertyContainerCustomization"
+
+namespace UE::ConcertReplicationScriptingEditor
+{
+	TSharedRef<IPropertyTypeCustomization> FConcertPropertyContainerCustomization::MakeInstance(FClassRememberer* ClassRememberer)
+	{
+		check(ClassRememberer);
+		return MakeShared<FConcertPropertyContainerCustomization>(*ClassRememberer);
+	}
+
+	FConcertPropertyContainerCustomization::~FConcertPropertyContainerCustomization()
+	{
+		FCoreUObjectDelegates::OnObjectTransacted.RemoveAll(this);
+	}
+
+	void FConcertPropertyContainerCustomization::CustomizeHeader(TSharedRef<IPropertyHandle> StructPropertyHandle, FDetailWidgetRow& HeaderRow, IPropertyTypeCustomizationUtils& StructCustomizationUtils)
+	{
+		PropertyHandle = StructPropertyHandle;
+		PropertyHandle->SetOnPropertyValueChanged(FSimpleDelegate::CreateSP(this, &FConcertPropertyContainerCustomization::RefreshPropertiesAndUpdateUI));
+		PropertyHandle->SetOnChildPropertyValueChanged(FSimpleDelegate::CreateSP(this, &FConcertPropertyContainerCustomization::RefreshPropertiesAndUpdateUI));
+		FCoreUObjectDelegates::OnObjectTransacted.AddSP(this, &FConcertPropertyContainerCustomization::OnObjectTransacted);
+		RefreshProperties();
+		
+		HeaderRow
+		.NameContent()
+		[
+			StructPropertyHandle->CreatePropertyNameWidget()
+		]
+		.ValueContent()
+		[
+			SNew(SBox)
+			.Padding(FMargin(0,2,0,1))
+			[
+				SAssignNew(PropertyComboBox, SConcertPropertyChainCombo)
+				.InitialClassSelection(ClassRememberer.GetLastUsedContextClassFor(*PropertyHandle))
+				.IsEditable(!PropertyHandle->IsEditConst())
+				.ContainedProperties(&Properties)
+				.OnClassChanged_Lambda([this](const UClass* Class)
+				{
+					ClassRememberer.OnUseClass(*PropertyHandle, Class);
+				})
+				.OnPropertySelectionChanged(this, &FConcertPropertyContainerCustomization::OnPropertySelectionChanged)
+			]
+		];
+	}
+
+	void FConcertPropertyContainerCustomization::RefreshProperties()
+	{
+		TArray<const void*> RawStructData;
+		PropertyHandle->AccessRawData(RawStructData);
+
+		Properties.Reset();
+		bHasMultipleValues = false;
+
+		// Sets are easier to compare than arrays
+		using FChainSet = TSet<FConcertPropertyChainWrapper, DefaultKeyFuncs<FConcertPropertyChainWrapper>, TInlineSetAllocator<64>>;
+		TOptional<FChainSet> SharedValue;
+		
+		for (int32 Idx = 0; Idx < RawStructData.Num(); ++Idx)
+		{
+			if (!RawStructData[Idx])
+			{
+				continue;
+			}
+
+			const FConcertPropertyChainWrapperContainer& CurrentValue = *(FConcertPropertyChainWrapperContainer*)RawStructData[Idx];
+			Algo::Transform(CurrentValue.PropertyChains, Properties, [](const FConcertPropertyChainWrapper& Wrapper){ return Wrapper.PropertyChain; });
+
+			FChainSet ValueAsSet(CurrentValue.PropertyChains);
+			if (!SharedValue)
+			{
+				SharedValue.Emplace(ValueAsSet);
+			}
+
+			const bool bValueIsDifferent = SharedValue->Num() != ValueAsSet.Num() || !SharedValue->Includes(ValueAsSet);
+			bHasMultipleValues |= bValueIsDifferent;
+		}
+	}
+
+	void FConcertPropertyContainerCustomization::RefreshPropertiesAndUpdateUI()
+	{
+		RefreshProperties();
+		PropertyComboBox->RefreshPropertyContent();
+	}
+
+	void FConcertPropertyContainerCustomization::OnPropertySelectionChanged(const FConcertPropertyChain& Property, bool bIsSelected)
+	{
+		FScopedTransaction Transaction(LOCTEXT("ChangeProperty", "Modify property container"));
+		PropertyHandle->NotifyPreChange();
+		ON_SCOPE_EXIT{ PropertyHandle->NotifyPostChange(EPropertyChangeType::ValueSet); };
+		
+		if (bIsSelected)
+		{
+			Properties.Add(Property);
+		}
+		else
+		{
+			Properties.Remove(Property);
+		}
+		
+		// And this updates the underlying property data
+		TArray<void*> RawStructData;
+		PropertyHandle->AccessRawData(RawStructData);
+		FConcertPropertyChainWrapperContainer NewValue;
+		Algo::Transform(Properties, NewValue.PropertyChains, [](const FConcertPropertyChain& PropertyChain){ return FConcertPropertyChainWrapper{ PropertyChain }; });
+		for (int32 Idx = 0; Idx < RawStructData.Num(); ++Idx)
+		{
+			FConcertPropertyChainWrapperContainer& Value = *(FConcertPropertyChainWrapperContainer*)RawStructData[Idx];
+			Value = NewValue;
+		}
+		
+		// This updates the combo button's content. 
+		PropertyComboBox->RefreshPropertyContent();
+	}
+
+	void FConcertPropertyContainerCustomization::OnObjectTransacted(UObject* Object, const FTransactionObjectEvent& TransactionObjectEvent)
+	{
+		TArray<UObject*> CustomizedObjects;
+		PropertyHandle->GetOuterObjects(CustomizedObjects);
+		if (CustomizedObjects.Contains(Object))
+		{
+			RefreshPropertiesAndUpdateUI();
+		}
+	}
+}
+
+#undef LOCTEXT_NAMESPACE
\ No newline at end of file
diff --git a/Engine/Plugins/Developer/Concert/ConcertScripting/ConcertReplicationScripting/Source/ConcertReplicationScriptingEditor/Private/Customization/ConcertPropertyContainerCustomization.h b/Engine/Plugins/Developer/Concert/ConcertScripting/ConcertReplicationScripting/Source/ConcertReplicationScriptingEditor/Private/Customization/ConcertPropertyContainerCustomization.h
new file mode 100644
--- /dev/null
+++ b/Engine/Plugins/Developer/Concert/ConcertScripting/ConcertReplicationScripting/Source/ConcertReplicationScriptingEditor/Private/Customization/ConcertPropertyContainerCustomization.h
@@ -0,0 +1,61 @@
+// Copyright Epic Games, Inc. All Rights Reserved.
+
+#pragma once
+
+#include "IPropertyTypeCustomization.h"
+#include "Widgets/SConcertPropertyChainCombo.h"
+
+struct FConcertPropertyChain;
+class IDetailChildrenBuilder;
+
+namespace UE::ConcertReplicationScriptingEditor
+{
+	class FClassRememberer;
+
+	/** Shows a combo button that allows selecting a property from a class. */
+	class FConcertPropertyContainerCustomization : public IPropertyTypeCustomization
+	{
+		template <typename ObjectType, ESPMode Mode>
+		friend class SharedPointerInternals::TIntrusiveReferenceController;
+	public:
+	
+		static TSharedRef<IPropertyTypeCustomization> MakeInstance(FClassRememberer* ClassRememberer);
+		
+		virtual ~FConcertPropertyContainerCustomization() override;
+
+		//~ Begin IPropertyTypeCustomization Interface
+		virtual void CustomizeHeader(TSharedRef<IPropertyHandle> StructPropertyHandle, FDetailWidgetRow& HeaderRow, IPropertyTypeCustomizationUtils& StructCustomizationUtils) override;
+		virtual void CustomizeChildren(TSharedRef<IPropertyHandle> InStructPropertyHandle, IDetailChildrenBuilder& ChildBuilder, IPropertyTypeCustomizationUtils& StructCustomizationUtils) override {}
+		//~ End IPropertyTypeCustomization Interface
+
+	private:
+		
+		FConcertPropertyContainerCustomization(FClassRememberer& InClassRememberer)
+			: ClassRememberer(InClassRememberer)
+		{}
+
+		/** Customizations use this cache so users do not constantly have to re-select the class in the drop-down menus. */
+		FClassRememberer& ClassRememberer;
+		
+		/** Handle to the property containing the FConcertPropertyChainWrapper struct. */
+		TSharedPtr<IPropertyHandle> PropertyHandle;
+
+		/** The UI displayed for the property. */
+		TSharedPtr<SConcertPropertyChainCombo> PropertyComboBox;
+
+		/**
+		 * Dummy passed in to the SConcertPropertyChainCombo.
+		 * Holds all the values from the different property sources.
+		 */
+		TSet<FConcertPropertyChain> Properties;
+		bool bHasMultipleValues = false;
+
+		void RefreshProperties();
+		void RefreshPropertiesAndUpdateUI();
+		
+		void OnPropertySelectionChanged(const FConcertPropertyChain& Property, bool bIsSelected);
+		void OnObjectTransacted(UObject* Object, const FTransactionObjectEvent& TransactionObjectEvent);
+	};
+}
+
+
diff --git a/Engine/Plugins/Developer/Concert/ConcertScripting/ConcertReplicationScripting/Source/ConcertReplicationScriptingEditor/Private/Customization/ConcertPropertyCustomization.cpp b/Engine/Plugins/Developer/Concert/ConcertScripting/ConcertReplicationScripting/Source/ConcertReplicationScriptingEditor/Private/Customization/ConcertPropertyCustomization.cpp
--- a/Engine/Plugins/Developer/Concert/ConcertScripting/ConcertReplicationScripting/Source/ConcertReplicationScriptingEditor/Private/Customization/ConcertPropertyCustomization.cpp
+++ b/Engine/Plugins/Developer/Concert/ConcertScripting/ConcertReplicationScripting/Source/ConcertReplicationScriptingEditor/Private/Customization/ConcertPropertyCustomization.cpp
@@ -3,35 +3,34 @@
 #include "ConcertPropertyCustomization.h"
 
 #include "ConcertPropertyChainWrapper.h"
+#include "Util/ClassRememberer.h"
 #include "Widgets/SConcertPropertyChainCombo.h"
 
 #include "DetailWidgetRow.h"
+#include "ScopedTransaction.h"
+
+#define LOCTEXT_NAMESPACE "FConcertPropertyCustomization"
 
 namespace UE::ConcertReplicationScriptingEditor
 {
-	/** Keeps track of the last class that was used for searching properties in SConcertPropertyChainCombo. */
-	static TWeakObjectPtr<const UClass> LastSelectedClass;
-	/**
-	 * Maps every property path to the last used class. This why when a user works with multiple properties, each one will remember the last used class for convenience.
-	 * This leaks memory but a user must edit a lot of properties to really be noticeable... should be ok.
-	 */
-	static TMap<FString, TWeakObjectPtr<const UClass>> PropertyToLastUsedClass;
-	
-	TSharedRef<IPropertyTypeCustomization> FConcertPropertyCustomization::MakeInstance()
+	TSharedRef<IPropertyTypeCustomization> FConcertPropertyCustomization::MakeInstance(FClassRememberer* ClassRememberer)
+	{
+		check(ClassRememberer);
+		return MakeShared<FConcertPropertyCustomization>(*ClassRememberer);
+	}
+
+	FConcertPropertyCustomization::~FConcertPropertyCustomization()
 	{
-		return MakeShared<FConcertPropertyCustomization>();
+		FCoreUObjectDelegates::OnObjectTransacted.RemoveAll(this);
 	}
 
 	void FConcertPropertyCustomization::CustomizeHeader(TSharedRef<IPropertyHandle> StructPropertyHandle, FDetailWidgetRow& HeaderRow, IPropertyTypeCustomizationUtils& StructCustomizationUtils)
 	{
 		PropertyHandle = StructPropertyHandle;
 		PropertyHandle->SetOnPropertyValueChanged(FSimpleDelegate::CreateSP(this, &FConcertPropertyCustomization::RefreshPropertiesAndUpdateUI));
+		FCoreUObjectDelegates::OnObjectTransacted.AddSP(this, &FConcertPropertyCustomization::OnObjectTransacted);
 		RefreshProperties();
 		
-		const FString PropertyPath(PropertyHandle->GetPropertyPath());
-		const TWeakObjectPtr<const UClass>* LastBoundClass = PropertyToLastUsedClass.Find(PropertyPath);
-		const TWeakObjectPtr<const UClass> ClassToUse = LastBoundClass ? *LastBoundClass : LastSelectedClass;
-		
 		HeaderRow
 		.NameContent()
 		[
@@ -43,15 +42,12 @@ namespace UE::ConcertReplicationScriptingEditor
 			.Padding(FMargin(0,2,0,1))
 			[
 				SAssignNew(PropertyComboBox, SConcertPropertyChainCombo)
-				.InitialClassSelection(ClassToUse.IsValid() ? ClassToUse.Get() : nullptr)
+				.InitialClassSelection(ClassRememberer.GetLastUsedContextClassFor(*PropertyHandle))
 				.IsEditable(!PropertyHandle->IsEditConst())
 				.ContainedProperties(&Properties)
 				.OnClassChanged_Lambda([this](const UClass* Class)
 				{
-					LastSelectedClass = Class;
-					
-					const FString PropertyPath(PropertyHandle->GetPropertyPath());
-					PropertyToLastUsedClass.Add(PropertyPath, Class);
+					ClassRememberer.OnUseClass(*PropertyHandle, Class);
 				})
 				.OnPropertySelectionChanged(this, &FConcertPropertyCustomization::OnPropertySelectionChanged)
 			]
@@ -95,6 +91,10 @@ namespace UE::ConcertReplicationScriptingEditor
 
 	void FConcertPropertyCustomization::OnPropertySelectionChanged(const FConcertPropertyChain& Property, bool bIsSelected)
 	{
+		FScopedTransaction Transaction(LOCTEXT("ChangeProperty", "Modify property"));
+		PropertyHandle->NotifyPreChange();
+		ON_SCOPE_EXIT{ PropertyHandle->NotifyPostChange(EPropertyChangeType::ValueSet); };
+		
 		// And this updates the underlying property data
 		TArray<void*> RawStructData;
 		PropertyHandle->AccessRawData(RawStructData);
@@ -110,4 +110,16 @@ namespace UE::ConcertReplicationScriptingEditor
 		// ... and this updates the combo button's content. 
 		PropertyComboBox->RefreshPropertyContent();
 	}
+
+	void FConcertPropertyCustomization::OnObjectTransacted(UObject* Object, const FTransactionObjectEvent& TransactionObjectEvent)
+	{
+		TArray<UObject*> CustomizedObjects;
+		PropertyHandle->GetOuterObjects(CustomizedObjects);
+		if (CustomizedObjects.Contains(Object))
+		{
+			RefreshPropertiesAndUpdateUI();
+		}
+	}
 }
+
+#undef LOCTEXT_NAMESPACE
\ No newline at end of file
diff --git a/Engine/Plugins/Developer/Concert/ConcertScripting/ConcertReplicationScripting/Source/ConcertReplicationScriptingEditor/Private/Customization/ConcertPropertyCustomization.h b/Engine/Plugins/Developer/Concert/ConcertScripting/ConcertReplicationScripting/Source/ConcertReplicationScriptingEditor/Private/Customization/ConcertPropertyCustomization.h
--- a/Engine/Plugins/Developer/Concert/ConcertScripting/ConcertReplicationScripting/Source/ConcertReplicationScriptingEditor/Private/Customization/ConcertPropertyCustomization.h
+++ b/Engine/Plugins/Developer/Concert/ConcertScripting/ConcertReplicationScripting/Source/ConcertReplicationScriptingEditor/Private/Customization/ConcertPropertyCustomization.h
@@ -10,12 +10,18 @@ class IDetailChildrenBuilder;
 
 namespace UE::ConcertReplicationScriptingEditor
 {
+	class FClassRememberer;
+
 	/** Shows a combo button that allows selecting a property from a class. */
 	class FConcertPropertyCustomization : public IPropertyTypeCustomization
 	{
+		template <typename ObjectType, ESPMode Mode>
+		friend class SharedPointerInternals::TIntrusiveReferenceController;
 	public:
-	
-		static TSharedRef<IPropertyTypeCustomization> MakeInstance();
+		
+		static TSharedRef<IPropertyTypeCustomization> MakeInstance(FClassRememberer* ClassRememberer);
+
+		virtual ~FConcertPropertyCustomization() override;
 
 		//~ Begin IPropertyTypeCustomization Interface
 		virtual void CustomizeHeader(TSharedRef<IPropertyHandle> StructPropertyHandle, FDetailWidgetRow& HeaderRow, IPropertyTypeCustomizationUtils& StructCustomizationUtils) override;
@@ -24,13 +30,23 @@ namespace UE::ConcertReplicationScriptingEditor
 
 	private:
 
+		FConcertPropertyCustomization(FClassRememberer& InClassRememberer)
+			: ClassRememberer(InClassRememberer)
+		{}
+
+		/** Customizations use this cache so users do not constantly have to re-select the class in the drop-down menus. */
+		FClassRememberer& ClassRememberer;
+
 		/** Handle to the property containing the FConcertPropertyChainWrapper struct. */
 		TSharedPtr<IPropertyHandle> PropertyHandle;
 
 		/** The UI displayed for the property. */
 		TSharedPtr<SConcertPropertyChainCombo> PropertyComboBox;
 
-		/** Dummy passed in to the SConcertPropertyChainCombo. Holds all the values from the different property sources. */
+		/**
+		 * Dummy passed in to the SConcertPropertyChainCombo.
+		 * Holds all the values from the different property sources.
+		 */
 		TSet<FConcertPropertyChain> Properties;
 		bool bHasMultipleValues = false;
 
@@ -38,6 +54,7 @@ namespace UE::ConcertReplicationScriptingEditor
 		void RefreshPropertiesAndUpdateUI();
 		
 		void OnPropertySelectionChanged(const FConcertPropertyChain& Property, bool bIsSelected);
+		void OnObjectTransacted(UObject* Object, const FTransactionObjectEvent& TransactionObjectEvent);
 	};
 }
 
diff --git a/Engine/Plugins/Developer/Concert/ConcertScripting/ConcertReplicationScripting/Source/ConcertReplicationScriptingEditor/Private/Util/ClassRememberer.cpp b/Engine/Plugins/Developer/Concert/ConcertScripting/ConcertReplicationScripting/Source/ConcertReplicationScriptingEditor/Private/Util/ClassRememberer.cpp
new file mode 100644
--- /dev/null
+++ b/Engine/Plugins/Developer/Concert/ConcertScripting/ConcertReplicationScripting/Source/ConcertReplicationScriptingEditor/Private/Util/ClassRememberer.cpp
@@ -0,0 +1,26 @@
+// Copyright Epic Games, Inc. All Rights Reserved.
+
+#include "ClassRememberer.h"
+
+#include "PropertyHandle.h"
+
+namespace UE::ConcertReplicationScriptingEditor
+{
+	const UClass* FClassRememberer::GetLastUsedContextClassFor(IPropertyHandle& Property) const
+	{
+		const FString PropertyPath(Property.GetPropertyPath());
+		const TWeakObjectPtr<const UClass>* LastBoundClass = PropertyToLastUsedClass.Find(PropertyPath);
+		const TWeakObjectPtr<const UClass> ClassToUse = LastBoundClass ? *LastBoundClass : LastSelectedClass;
+		return ClassToUse.IsValid() ? ClassToUse.Get() : nullptr;
+	}
+	
+	void FClassRememberer::OnUseClass(IPropertyHandle& Property, const UClass* Class)
+	{
+		if (Class)
+		{
+			LastSelectedClass = Class;
+			const FString PropertyPath(Property.GetPropertyPath());
+			PropertyToLastUsedClass.Add(PropertyPath, Class);
+		}
+	}
+}
diff --git a/Engine/Plugins/Developer/Concert/ConcertScripting/ConcertReplicationScripting/Source/ConcertReplicationScriptingEditor/Private/Util/ClassRememberer.h b/Engine/Plugins/Developer/Concert/ConcertScripting/ConcertReplicationScripting/Source/ConcertReplicationScriptingEditor/Private/Util/ClassRememberer.h
new file mode 100644
--- /dev/null
+++ b/Engine/Plugins/Developer/Concert/ConcertScripting/ConcertReplicationScripting/Source/ConcertReplicationScriptingEditor/Private/Util/ClassRememberer.h
@@ -0,0 +1,36 @@
+// Copyright Epic Games, Inc. All Rights Reserved.
+
+#pragma once
+
+#include "UObject/WeakObjectPtr.h"
+#include "UObject/WeakObjectPtrTemplates.h"
+#include "Templates/UnrealTemplate.h"
+
+class IPropertyHandle;
+class UClass;
+
+namespace UE::ConcertReplicationScriptingEditor
+{
+	/** Remembers which classes were used based per customized property so users do not constantly have to re-select it in the drop-down menus. */
+	class FClassRememberer : public FNoncopyable
+	{
+	public:
+
+		/** Gets the class to use for searching properties in the FConcertPropertyChainWrapper or FConcertPropertyChainWrapperContainer. */
+		const UClass* GetLastUsedContextClassFor(IPropertyHandle& Property) const;
+		
+		/** Caches a class for the property. */
+		void OnUseClass(IPropertyHandle& Property, const UClass* Class);
+	
+	private:
+		
+		/** Keeps track of the last class that was used for searching properties in SConcertPropertyChainCombo. */
+		TWeakObjectPtr<const UClass> LastSelectedClass;
+		/**
+		 * Maps every property path to the last used class. This why when a user works with multiple properties, each one will remember the last used class for convenience.
+		 * This leaks memory but a user must edit a lot of properties to really be noticeable... should be ok.
+		 */
+		TMap<FString, TWeakObjectPtr<const UClass>> PropertyToLastUsedClass;
+	};
+}
+
diff --git a/Engine/Plugins/Developer/Concert/ConcertScripting/ConcertReplicationScripting/Source/ConcertReplicationScriptingEditor/Private/Widgets/SConcertPropertyChainChip.cpp b/Engine/Plugins/Developer/Concert/ConcertScripting/ConcertReplicationScripting/Source/ConcertReplicationScriptingEditor/Private/Widgets/SConcertPropertyChainChip.cpp
--- a/Engine/Plugins/Developer/Concert/ConcertScripting/ConcertReplicationScripting/Source/ConcertReplicationScriptingEditor/Private/Widgets/SConcertPropertyChainChip.cpp
+++ b/Engine/Plugins/Developer/Concert/ConcertScripting/ConcertReplicationScripting/Source/ConcertReplicationScriptingEditor/Private/Widgets/SConcertPropertyChainChip.cpp
@@ -11,6 +11,7 @@
 #include "Widgets/Input/SButton.h"
 #include "Widgets/Layout/SWidgetSwitcher.h"
 #include "Widgets/SBoxPanel.h"
+#include "Widgets/SNullWidget.h"
 #include "Widgets/Text/STextBlock.h"
 
 #define LOCTEXT_NAMESPACE "SConcertPropertyChainChip"
@@ -27,45 +28,58 @@ namespace UE::ConcertReplicationScriptingEditor
 
 			+SWidgetSwitcher::Slot()
 			[
-				SNew(SButton)
-				.ButtonStyle(FReplicationScriptingStyle::Get(), "ConcertProperty.ChipButton.Selected")
-				.OnClicked_Lambda([OnEditPressed = InArgs._OnEditPressed](){ OnEditPressed.ExecuteIfBound(); return FReply::Handled(); })
+				SNew(SHorizontalBox)
+				
+				+SHorizontalBox::Slot()
+				.AutoWidth()
 				[
-					SNew(SHorizontalBox)
-
-					+SHorizontalBox::Slot()
-					.AutoWidth()
-					.VAlign(VAlign_Center)
+					SNew(SButton)
+					.ButtonStyle(FReplicationScriptingStyle::Get(), "ConcertProperty.ChipButton.Selected")
+					.OnClicked_Lambda([OnEditPressed = InArgs._OnEditPressed](){ OnEditPressed.ExecuteIfBound(); return FReply::Handled(); })
 					[
-						SNew(STextBlock)
-						.Font(FAppStyle::GetFontStyle( TEXT("PropertyWindow.NormalFont")))
-						.Text(FText::FromString(PropertyChain.ToString()))
-					]
+						SNew(SHorizontalBox)
 
-					+SHorizontalBox::Slot()
-					.AutoWidth()
-					.VAlign(VAlign_Center)
-					[
-						SNew(SButton)
-						.Visibility(!PropertyChain.IsEmpty() && InArgs._ShowClearButton ? EVisibility::Visible : EVisibility::Collapsed)
+						+SHorizontalBox::Slot()
+						.AutoWidth()
 						.VAlign(VAlign_Center)
-						.HAlign(HAlign_Center)
-						.ButtonStyle(FReplicationScriptingStyle::Get(), "ConcertProperty.ChipClearButton")
-						.ToolTipText(LOCTEXT("Remove", "Remove"))
-						.ContentPadding(FMargin(2, 0, 0, 0))
-						.OnClicked_Lambda([OnClearPressed = InArgs._OnClearPressed]()
-						{
-							OnClearPressed.ExecuteIfBound();
-							return FReply::Handled();
-						})
 						[
-							SNew(SImage)
-							.ColorAndOpacity(FStyleColors::Foreground)
-							.Image(FAppStyle::GetBrush("Icons.X"))
-							.DesiredSizeOverride(FVector2D(12.0f, 12.0f))
+							SNew(STextBlock)
+							.Font(FAppStyle::GetFontStyle( TEXT("PropertyWindow.NormalFont")))
+							.Text(FText::FromString(PropertyChain.ToString()))
+						]
+
+						+SHorizontalBox::Slot()
+						.AutoWidth()
+						.VAlign(VAlign_Center)
+						[
+							SNew(SButton)
+							.Visibility(!PropertyChain.IsEmpty() && InArgs._ShowClearButton ? EVisibility::Visible : EVisibility::Collapsed)
+							.VAlign(VAlign_Center)
+							.HAlign(HAlign_Center)
+							.ButtonStyle(FReplicationScriptingStyle::Get(), "ConcertProperty.ChipClearButton")
+							.ToolTipText(LOCTEXT("Remove", "Remove"))
+							.ContentPadding(FMargin(2, 0, 0, 0))
+							.OnClicked_Lambda([OnClearPressed = InArgs._OnClearPressed]()
+							{
+								OnClearPressed.ExecuteIfBound();
+								return FReply::Handled();
+							})
+							[
+								SNew(SImage)
+								.ColorAndOpacity(FStyleColors::Foreground)
+								.Image(FAppStyle::GetBrush("Icons.X"))
+								.DesiredSizeOverride(FVector2D(12.0f, 12.0f))
+							]
 						]
 					]
 				]
+				
+				// So that the chip does not take up the entire row and leaves a lot of empty space next to the remove button
+				+SHorizontalBox::Slot()
+				.FillWidth(1.f)
+				[
+					SNullWidget::NullWidget
+				]
 			]
 			
 			+SWidgetSwitcher::Slot()
diff --git a/Engine/Plugins/Developer/Concert/ConcertScripting/ConcertReplicationScripting/Source/ConcertReplicationScriptingEditor/Private/Widgets/SConcertPropertyChainCombo.cpp b/Engine/Plugins/Developer/Concert/ConcertScripting/ConcertReplicationScripting/Source/ConcertReplicationScriptingEditor/Private/Widgets/SConcertPropertyChainCombo.cpp
--- a/Engine/Plugins/Developer/Concert/ConcertScripting/ConcertReplicationScripting/Source/ConcertReplicationScriptingEditor/Private/Widgets/SConcertPropertyChainCombo.cpp
+++ b/Engine/Plugins/Developer/Concert/ConcertScripting/ConcertReplicationScripting/Source/ConcertReplicationScriptingEditor/Private/Widgets/SConcertPropertyChainCombo.cpp
@@ -32,7 +32,14 @@ namespace UE::ConcertReplicationScriptingEditor
 			.ButtonContent()
 			[
 				SNew(SWidgetSwitcher)
-				.WidgetIndex_Lambda([this](){ return HasMultipleValuesAttribute.Get() ? 0 : 1; })
+				.WidgetIndex_Lambda([this]()
+				{
+					return HasMultipleValuesAttribute.Get()
+						? 0
+						: DisplayedProperties.IsEmpty()
+							? 1
+							: 2;
+				})
 				
 				+SWidgetSwitcher::Slot()
 				[
@@ -40,6 +47,12 @@ namespace UE::ConcertReplicationScriptingEditor
 					.Text(LOCTEXT("MultipleValues", "Multiple Values"))
 				]
 				+SWidgetSwitcher::Slot()
+				[
+					SNew(SConcertPropertyChainChip)
+					.DisplayedProperty(FConcertPropertyChain{})
+					.ShowClearButton(false)
+				]
+				+SWidgetSwitcher::Slot()
 				[
 					SAssignNew(PropertyListView, SListView<TSharedPtr<FConcertPropertyChain>>)
 					.ListItemsSource(&DisplayedProperties)
diff --git a/Engine/Plugins/Developer/Concert/ConcertScripting/ConcertReplicationScripting/Source/ConcertReplicationScriptingEditor/Private/Widgets/SConcertPropertyChainCombo.h b/Engine/Plugins/Developer/Concert/ConcertScripting/ConcertReplicationScripting/Source/ConcertReplicationScriptingEditor/Private/Widgets/SConcertPropertyChainCombo.h
--- a/Engine/Plugins/Developer/Concert/ConcertScripting/ConcertReplicationScripting/Source/ConcertReplicationScriptingEditor/Private/Widgets/SConcertPropertyChainCombo.h
+++ b/Engine/Plugins/Developer/Concert/ConcertScripting/ConcertReplicationScripting/Source/ConcertReplicationScriptingEditor/Private/Widgets/SConcertPropertyChainCombo.h
@@ -25,7 +25,12 @@ namespace UE::ConcertReplicationScriptingEditor
 		SLATE_BEGIN_ARGS(SConcertPropertyChainCombo)
 			: _InitialClassSelection(nullptr)
 		{}
-			/** The properties that are supposed to be displayed inside of the button */
+			/**
+			 * The properties that are supposed to be displayed inside of the button.
+			 * 
+			 * This must always contain at least one property. If there are "no" properties, this should contain the empty, default constructed FConcertPropertyChain.
+			 * This requirement is needed so the button shows the "Empty" widget.
+			 */
 			SLATE_ARGUMENT(const TSet<FConcertPropertyChain>*, ContainedProperties)
 
 			/** Whether the UI should allow editing */
diff --git a/Engine/Plugins/Developer/Concert/ConcertScripting/ConcertReplicationScripting/Source/ConcertReplicationScriptingEditor/Private/Widgets/SConcertPropertyChainPicker.cpp b/Engine/Plugins/Developer/Concert/ConcertScripting/ConcertReplicationScripting/Source/ConcertReplicationScriptingEditor/Private/Widgets/SConcertPropertyChainPicker.cpp
--- a/Engine/Plugins/Developer/Concert/ConcertScripting/ConcertReplicationScripting/Source/ConcertReplicationScriptingEditor/Private/Widgets/SConcertPropertyChainPicker.cpp
+++ b/Engine/Plugins/Developer/Concert/ConcertScripting/ConcertReplicationScripting/Source/ConcertReplicationScriptingEditor/Private/Widgets/SConcertPropertyChainPicker.cpp
@@ -139,6 +139,9 @@ namespace UE::ConcertReplicationScriptingEditor
 	void SConcertPropertyChainPicker::OnPropertySelected(const FConcertPropertyChain& ConcertPropertyChain, bool bIsSelected)
 	{
 		OnSelectedPropertiesChangedDelegate.ExecuteIfBound(ConcertPropertyChain, bIsSelected);
+
+		// The checkbox state may have changed request a resort.
+		TreeView->RequestResortForColumn(PropertySelectionCheckboxColumnId);
 	}
 }
 
