Fix non-unity error

diff --git a/Engine/Plugins/Developer/Concert/ConcertScripting/ConcertReplicationScripting/Source/ConcertReplicationScripting/Public/ConcertPropertyChainWrapperContainer.h b/Engine/Plugins/Developer/Concert/ConcertScripting/ConcertReplicationScripting/Source/ConcertReplicationScripting/Public/ConcertPropertyChainWrapperContainer.h
--- a/Engine/Plugins/Developer/Concert/ConcertScripting/ConcertReplicationScripting/Source/ConcertReplicationScripting/Public/ConcertPropertyChainWrapperContainer.h
+++ b/Engine/Plugins/Developer/Concert/ConcertScripting/ConcertReplicationScripting/Source/ConcertReplicationScripting/Public/ConcertPropertyChainWrapperContainer.h
@@ -2,7 +2,8 @@
 
 #pragma once
 
-#include "Replication/Data/ConcertPropertySelection.h"
+#include "ConcertPropertyChainWrapper.h"
+#include "Containers/Array.h"
 #include "ConcertPropertyChainWrapperContainer.generated.h"
 
 /** Special array of property chains. Used to allow for detail customization. */
